GUBP Node Shared_LabelPromotable built from changelist 2505032

diff --git a/Engine/Build/Receipts/CrashReportClient-Linux-Shippingx86_64-unknown-linux-gnu.target.xml b/Engine/Build/Receipts/CrashReportClient-Linux-Shippingx86_64-unknown-linux-gnu.target.xml
new file mode 100644
--- /dev/null
+++ b/Engine/Build/Receipts/CrashReportClient-Linux-Shippingx86_64-unknown-linux-gnu.target.xml
@@ -0,0 +1,7 @@
+<?xml version="1.0" encoding="utf-8"?>
+<BuildReceipt xmlns:xsd="http://www.w3.org/2001/XMLSchema" xmlns:xsi="http://www.w3.org/2001/XMLSchema-instance">
+  <BuildProducts>
+    <BuildProduct Path="$(EngineDir)\Binaries\Linux\CrashReportClient" Type="Executable" />
+  </BuildProducts>
+  <RuntimeDependencies />
+</BuildReceipt>
\ No newline at end of file
diff --git a/Engine/Build/Receipts/PS4MapFileUtil-Win64-Development.target.xml b/Engine/Build/Receipts/PS4MapFileUtil-Win64-Development.target.xml
new file mode 100644
--- /dev/null
+++ b/Engine/Build/Receipts/PS4MapFileUtil-Win64-Development.target.xml
@@ -0,0 +1,15 @@
+<?xml version="1.0" encoding="utf-8"?>
+<BuildReceipt xmlns:xsd="http://www.w3.org/2001/XMLSchema" xmlns:xsi="http://www.w3.org/2001/XMLSchema-instance">
+  <BuildProducts>
+    <BuildProduct Path="$(EngineDir)\Binaries\Win64\PS4\PS4MapFileUtil.exe" Type="Executable" />
+    <BuildProduct Path="$(EngineDir)\Binaries\Win64\PS4\PS4MapFileUtil.pdb" Type="SymbolFile" />
+  </BuildProducts>
+  <RuntimeDependencies>
+    <RuntimeDependency Path="$(EngineDir)/Binaries/ThirdParty/ICU/icu4c-53_1/Win64/VS2013/icudt53.dll" />
+    <RuntimeDependency Path="$(EngineDir)/Binaries/ThirdParty/ICU/icu4c-53_1/Win64/VS2013/icuuc53.dll" />
+    <RuntimeDependency Path="$(EngineDir)/Binaries/ThirdParty/ICU/icu4c-53_1/Win64/VS2013/icuin53.dll" />
+    <RuntimeDependency Path="$(EngineDir)/Binaries/ThirdParty/ICU/icu4c-53_1/Win64/VS2013/icule53.dll" />
+    <RuntimeDependency Path="$(EngineDir)/Binaries/ThirdParty/ICU/icu4c-53_1/Win64/VS2013/iculx53.dll" />
+    <RuntimeDependency Path="$(EngineDir)/Binaries/ThirdParty/ICU/icu4c-53_1/Win64/VS2013/icuio53.dll" />
+  </RuntimeDependencies>
+</BuildReceipt>
\ No newline at end of file
diff --git a/Engine/Build/Receipts/UE4Client-IOS-Development.target.xml b/Engine/Build/Receipts/UE4Client-IOS-Development.target.xml
new file mode 100644
--- /dev/null
+++ b/Engine/Build/Receipts/UE4Client-IOS-Development.target.xml
@@ -0,0 +1,18 @@
+<?xml version="1.0" encoding="utf-8"?>
+<BuildReceipt xmlns:xsi="http://www.w3.org/2001/XMLSchema-instance" xmlns:xsd="http://www.w3.org/2001/XMLSchema">
+  <BuildProducts>
+    <BuildProduct Path="$(EngineDir)/Binaries/IOS/UE4Client" Type="Executable" />
+    <BuildProduct Path="$(EngineDir)/Binaries/IOS/UE4Client.stub" Type="Executable" />
+    <BuildProduct Path="$(EngineDir)/Plugins/2D/Paper2D/Binaries/IOS/UE4Client-Paper2D-Static.a" Type="StaticLibrary" />
+    <BuildProduct Path="$(EngineDir)/Plugins/Experimental/CharacterAI/Binaries/IOS/UE4Client-CharacterAI-Static.a" Type="StaticLibrary" />
+    <BuildProduct Path="$(EngineDir)/Plugins/Media/AvfMedia/Binaries/IOS/UE4Client-AvfMedia-Static.a" Type="StaticLibrary" />
+    <BuildProduct Path="$(EngineDir)/Plugins/Messaging/UdpMessaging/Binaries/IOS/UE4Client-UdpMessaging-Static.a" Type="StaticLibrary" />
+    <BuildProduct Path="$(EngineDir)/Plugins/Runtime/AppleMoviePlayer/Binaries/IOS/UE4Client-AppleMoviePlayer-Static.a" Type="StaticLibrary" />
+    <BuildProduct Path="$(EngineDir)/Plugins/Runtime/CableComponent/Binaries/IOS/UE4Client-CableComponent-Static.a" Type="StaticLibrary" />
+    <BuildProduct Path="$(EngineDir)/Plugins/Runtime/CustomMeshComponent/Binaries/IOS/UE4Client-CustomMeshComponent-Static.a" Type="StaticLibrary" />
+    <BuildProduct Path="$(EngineDir)/Plugins/Runtime/ExampleDeviceProfileSelector/Binaries/IOS/UE4Client-ExampleDeviceProfileSelector-Static.a" Type="StaticLibrary" />
+    <BuildProduct Path="$(EngineDir)/Plugins/Runtime/IOSDeviceProfileSelector/Binaries/IOS/UE4Client-IOSDeviceProfileSelector-Static.a" Type="StaticLibrary" />
+    <BuildProduct Path="$(EngineDir)/Plugins/Slate/SlateRemote/Binaries/IOS/UE4Client-SlateRemote-Static.a" Type="StaticLibrary" />
+  </BuildProducts>
+  <RuntimeDependencies />
+</BuildReceipt>
\ No newline at end of file
diff --git a/Engine/Build/Receipts/UE4Client-Linux-Developmentx86_64-unknown-linux-gnu.target.xml b/Engine/Build/Receipts/UE4Client-Linux-Developmentx86_64-unknown-linux-gnu.target.xml
new file mode 100644
--- /dev/null
+++ b/Engine/Build/Receipts/UE4Client-Linux-Developmentx86_64-unknown-linux-gnu.target.xml
@@ -0,0 +1,14 @@
+<?xml version="1.0" encoding="utf-8"?>
+<BuildReceipt xmlns:xsd="http://www.w3.org/2001/XMLSchema" xmlns:xsi="http://www.w3.org/2001/XMLSchema-instance">
+  <BuildProducts>
+    <BuildProduct Path="$(EngineDir)\Binaries\Linux\UE4Client" Type="Executable" />
+    <BuildProduct Path="$(EngineDir)\Plugins\2D\Paper2D\Binaries\Linux\libUE4Client-Paper2D-Static.a" Type="StaticLibrary" />
+    <BuildProduct Path="$(EngineDir)\Plugins\Experimental\CharacterAI\Binaries\Linux\libUE4Client-CharacterAI-Static.a" Type="StaticLibrary" />
+    <BuildProduct Path="$(EngineDir)\Plugins\Messaging\UdpMessaging\Binaries\Linux\libUE4Client-UdpMessaging-Static.a" Type="StaticLibrary" />
+    <BuildProduct Path="$(EngineDir)\Plugins\Runtime\CableComponent\Binaries\Linux\libUE4Client-CableComponent-Static.a" Type="StaticLibrary" />
+    <BuildProduct Path="$(EngineDir)\Plugins\Runtime\CustomMeshComponent\Binaries\Linux\libUE4Client-CustomMeshComponent-Static.a" Type="StaticLibrary" />
+    <BuildProduct Path="$(EngineDir)\Plugins\Runtime\ExampleDeviceProfileSelector\Binaries\Linux\libUE4Client-ExampleDeviceProfileSelector-Static.a" Type="StaticLibrary" />
+    <BuildProduct Path="$(EngineDir)\Plugins\Slate\SlateRemote\Binaries\Linux\libUE4Client-SlateRemote-Static.a" Type="StaticLibrary" />
+  </BuildProducts>
+  <RuntimeDependencies />
+</BuildReceipt>
\ No newline at end of file
diff --git a/Engine/Build/Receipts/UE4Client-Mac-Development.target.xml b/Engine/Build/Receipts/UE4Client-Mac-Development.target.xml
new file mode 100644
--- /dev/null
+++ b/Engine/Build/Receipts/UE4Client-Mac-Development.target.xml
@@ -0,0 +1,39 @@
+<?xml version="1.0" encoding="utf-8"?>
+<BuildReceipt xmlns:xsi="http://www.w3.org/2001/XMLSchema-instance" xmlns:xsd="http://www.w3.org/2001/XMLSchema">
+  <BuildProducts>
+    <BuildProduct Path="$(EngineDir)/Binaries/Mac/UE4Client.app/Contents/MacOS/UE4Client" Type="Executable" />
+    <BuildProduct Path="$(EngineDir)/Binaries/Mac/UE4Client.app/Contents/MacOS/libogg.dylib" Type="DynamicLibrary" />
+    <BuildProduct Path="$(EngineDir)/Binaries/Mac/UE4Client.app/Contents/MacOS/libvorbis.dylib" Type="DynamicLibrary" />
+    <BuildProduct Path="$(EngineDir)/Binaries/Mac/UE4Client.app/Contents/Resources/RadioEffectUnit.component/Contents/Info.plist" Type="RequiredResource" />
+    <BuildProduct Path="$(EngineDir)/Binaries/Mac/UE4Client.app/Contents/Resources/RadioEffectUnit.component/Contents/MacOS/RadioEffectUnit" Type="RequiredResource" />
+    <BuildProduct Path="$(EngineDir)/Binaries/Mac/UE4Client.app/Contents/Resources/RadioEffectUnit.component/Contents/Resources/English.lproj/Localizable.strings" Type="RequiredResource" />
+    <BuildProduct Path="$(EngineDir)/Binaries/Mac/UE4Client.app/Contents/Info.plist" Type="RequiredResource" />
+    <BuildProduct Path="$(EngineDir)/Binaries/Mac/UE4Client.app/Contents/PkgInfo" Type="RequiredResource" />
+    <BuildProduct Path="$(EngineDir)/Binaries/Mac/UE4Client.app/Contents/Resources/UE4.icns" Type="RequiredResource" />
+    <BuildProduct Path="$(EngineDir)/Plugins/2D/Paper2D/Binaries/Mac/UE4Client-Paper2D-Static.a" Type="StaticLibrary" />
+    <BuildProduct Path="$(EngineDir)/Binaries/Mac/UE4Client.app/Contents/MacOS/libogg.dylib" Type="DynamicLibrary" />
+    <BuildProduct Path="$(EngineDir)/Binaries/Mac/UE4Client.app/Contents/MacOS/libvorbis.dylib" Type="DynamicLibrary" />
+    <BuildProduct Path="$(EngineDir)/Plugins/Experimental/CharacterAI/Binaries/Mac/UE4Client-CharacterAI-Static.a" Type="StaticLibrary" />
+    <BuildProduct Path="$(EngineDir)/Binaries/Mac/UE4Client.app/Contents/MacOS/libogg.dylib" Type="DynamicLibrary" />
+    <BuildProduct Path="$(EngineDir)/Binaries/Mac/UE4Client.app/Contents/MacOS/libvorbis.dylib" Type="DynamicLibrary" />
+    <BuildProduct Path="$(EngineDir)/Plugins/Media/AvfMedia/Binaries/Mac/UE4Client-AvfMedia-Static.a" Type="StaticLibrary" />
+    <BuildProduct Path="$(EngineDir)/Plugins/Messaging/UdpMessaging/Binaries/Mac/UE4Client-UdpMessaging-Static.a" Type="StaticLibrary" />
+    <BuildProduct Path="$(EngineDir)/Plugins/Runtime/AppleMoviePlayer/Binaries/Mac/UE4Client-AppleMoviePlayer-Static.a" Type="StaticLibrary" />
+    <BuildProduct Path="$(EngineDir)/Binaries/Mac/UE4Client.app/Contents/MacOS/libogg.dylib" Type="DynamicLibrary" />
+    <BuildProduct Path="$(EngineDir)/Binaries/Mac/UE4Client.app/Contents/MacOS/libvorbis.dylib" Type="DynamicLibrary" />
+    <BuildProduct Path="$(EngineDir)/Plugins/Runtime/CableComponent/Binaries/Mac/UE4Client-CableComponent-Static.a" Type="StaticLibrary" />
+    <BuildProduct Path="$(EngineDir)/Binaries/Mac/UE4Client.app/Contents/MacOS/libogg.dylib" Type="DynamicLibrary" />
+    <BuildProduct Path="$(EngineDir)/Binaries/Mac/UE4Client.app/Contents/MacOS/libvorbis.dylib" Type="DynamicLibrary" />
+    <BuildProduct Path="$(EngineDir)/Plugins/Runtime/CustomMeshComponent/Binaries/Mac/UE4Client-CustomMeshComponent-Static.a" Type="StaticLibrary" />
+    <BuildProduct Path="$(EngineDir)/Binaries/Mac/UE4Client.app/Contents/MacOS/libogg.dylib" Type="DynamicLibrary" />
+    <BuildProduct Path="$(EngineDir)/Binaries/Mac/UE4Client.app/Contents/MacOS/libvorbis.dylib" Type="DynamicLibrary" />
+    <BuildProduct Path="$(EngineDir)/Plugins/Runtime/ExampleDeviceProfileSelector/Binaries/Mac/UE4Client-ExampleDeviceProfileSelector-Static.a" Type="StaticLibrary" />
+    <BuildProduct Path="$(EngineDir)/Binaries/Mac/UE4Client.app/Contents/MacOS/libogg.dylib" Type="DynamicLibrary" />
+    <BuildProduct Path="$(EngineDir)/Binaries/Mac/UE4Client.app/Contents/MacOS/libvorbis.dylib" Type="DynamicLibrary" />
+    <BuildProduct Path="$(EngineDir)/Plugins/Runtime/OculusRift/Binaries/Mac/UE4Client-OculusRift-Static.a" Type="StaticLibrary" />
+    <BuildProduct Path="$(EngineDir)/Binaries/Mac/UE4Client.app/Contents/MacOS/libogg.dylib" Type="DynamicLibrary" />
+    <BuildProduct Path="$(EngineDir)/Binaries/Mac/UE4Client.app/Contents/MacOS/libvorbis.dylib" Type="DynamicLibrary" />
+    <BuildProduct Path="$(EngineDir)/Plugins/Slate/SlateRemote/Binaries/Mac/UE4Client-SlateRemote-Static.a" Type="StaticLibrary" />
+  </BuildProducts>
+  <RuntimeDependencies />
+</BuildReceipt>
\ No newline at end of file
diff --git a/Engine/Build/Receipts/UE4Client-Win32-Development.target.xml b/Engine/Build/Receipts/UE4Client-Win32-Development.target.xml
new file mode 100644
--- /dev/null
+++ b/Engine/Build/Receipts/UE4Client-Win32-Development.target.xml
@@ -0,0 +1,35 @@
+<?xml version="1.0" encoding="utf-8"?>
+<BuildReceipt xmlns:xsd="http://www.w3.org/2001/XMLSchema" xmlns:xsi="http://www.w3.org/2001/XMLSchema-instance">
+  <BuildProducts>
+    <BuildProduct Path="$(EngineDir)\Binaries\Win32\UE4Client.exe" Type="Executable" />
+    <BuildProduct Path="$(EngineDir)\Binaries\Win32\UE4Client.pdb" Type="SymbolFile" />
+    <BuildProduct Path="$(EngineDir)\Plugins\2D\Paper2D\Binaries\Win32\UE4Client-Paper2D-Static.lib" Type="StaticLibrary" />
+    <BuildProduct Path="$(EngineDir)\Plugins\Experimental\CharacterAI\Binaries\Win32\UE4Client-CharacterAI-Static.lib" Type="StaticLibrary" />
+    <BuildProduct Path="$(EngineDir)\Plugins\Media\WmfMedia\Binaries\Win32\UE4Client-WmfMedia-Static.lib" Type="StaticLibrary" />
+    <BuildProduct Path="$(EngineDir)\Plugins\Messaging\UdpMessaging\Binaries\Win32\UE4Client-UdpMessaging-Static.lib" Type="StaticLibrary" />
+    <BuildProduct Path="$(EngineDir)\Plugins\Runtime\CableComponent\Binaries\Win32\UE4Client-CableComponent-Static.lib" Type="StaticLibrary" />
+    <BuildProduct Path="$(EngineDir)\Plugins\Runtime\CustomMeshComponent\Binaries\Win32\UE4Client-CustomMeshComponent-Static.lib" Type="StaticLibrary" />
+    <BuildProduct Path="$(EngineDir)\Plugins\Runtime\ExampleDeviceProfileSelector\Binaries\Win32\UE4Client-ExampleDeviceProfileSelector-Static.lib" Type="StaticLibrary" />
+    <BuildProduct Path="$(EngineDir)\Plugins\Runtime\OculusRift\Binaries\Win32\UE4Client-OculusRift-Static.lib" Type="StaticLibrary" />
+    <BuildProduct Path="$(EngineDir)\Plugins\Runtime\WindowsMoviePlayer\Binaries\Win32\UE4Client-WindowsMoviePlayer-Static.lib" Type="StaticLibrary" />
+    <BuildProduct Path="$(EngineDir)\Plugins\Slate\SlateRemote\Binaries\Win32\UE4Client-SlateRemote-Static.lib" Type="StaticLibrary" />
+  </BuildProducts>
+  <RuntimeDependencies>
+    <RuntimeDependency Path="$(EngineDir)/Binaries/ThirdParty/PhysX/PhysX-3.3/Win64/VS2013/PhysX3PROFILE_x86.dll" />
+    <RuntimeDependency Path="$(EngineDir)/Binaries/ThirdParty/PhysX/PhysX-3.3/Win64/VS2013/PhysX3GpuPROFILE_x86.dll" />
+    <RuntimeDependency Path="$(EngineDir)/Binaries/ThirdParty/PhysX/PhysX-3.3/Win64/VS2013/PhysX3CommonPROFILE_x86.dll" />
+    <RuntimeDependency Path="$(EngineDir)/Binaries/ThirdParty/PhysX/PhysX-3.3/Win64/VS2013/PhysX3CookingPROFILE_x86.dll" />
+    <RuntimeDependency Path="$(EngineDir)/Binaries/ThirdParty/PhysX/PhysX-3.3/Win64/VS2013/nvToolsExt64_1.dll" />
+    <RuntimeDependency Path="$(EngineDir)/Binaries/ThirdParty/PhysX/PhysX-3.3/Win64/VS2013/PhysXDevice.dll" />
+    <RuntimeDependency Path="$(EngineDir)/Binaries/ThirdParty/PhysX/APEX-1.3/Win32/VS2013/APEX_ClothingPROFILE_x86.dll" />
+    <RuntimeDependency Path="$(EngineDir)/Binaries/ThirdParty/PhysX/APEX-1.3/Win32/VS2013/APEX_DestructiblePROFILE_x86.dll" />
+    <RuntimeDependency Path="$(EngineDir)/Binaries/ThirdParty/PhysX/APEX-1.3/Win32/VS2013/APEX_LegacyPROFILE_x86.dll" />
+    <RuntimeDependency Path="$(EngineDir)/Binaries/ThirdParty/PhysX/APEX-1.3/Win32/VS2013/APEX_LoaderPROFILE_x86.dll" />
+    <RuntimeDependency Path="$(EngineDir)/Binaries/ThirdParty/PhysX/APEX-1.3/Win32/VS2013/APEX_ParticlesPROFILE_x86.dll" />
+    <RuntimeDependency Path="$(EngineDir)/Binaries/ThirdParty/PhysX/APEX-1.3/Win32/VS2013/ApexFrameworkPROFILE_x86.dll" />
+    <RuntimeDependency Path="$(EngineDir)/Binaries/ThirdParty/Ogg/Win32/VS2013/libogg.dll" />
+    <RuntimeDependency Path="$(EngineDir)/Binaries/ThirdParty/Vorbis/Win32/VS2013/libvorbis.dll" />
+    <RuntimeDependency Path="$(EngineDir)/Binaries/ThirdParty/Vorbis/Win32/VS2013/libvorbisfile.dll" />
+    <RuntimeDependency Path="$(EngineDir)/Binaries/ThirdParty/Steamworks/Steamv132/Win32/steam_api.dll" />
+  </RuntimeDependencies>
+</BuildReceipt>
\ No newline at end of file
diff --git a/Engine/Build/Receipts/UE4Client-Win64-Development.target.xml b/Engine/Build/Receipts/UE4Client-Win64-Development.target.xml
new file mode 100644
--- /dev/null
+++ b/Engine/Build/Receipts/UE4Client-Win64-Development.target.xml
@@ -0,0 +1,35 @@
+<?xml version="1.0" encoding="utf-8"?>
+<BuildReceipt xmlns:xsd="http://www.w3.org/2001/XMLSchema" xmlns:xsi="http://www.w3.org/2001/XMLSchema-instance">
+  <BuildProducts>
+    <BuildProduct Path="$(EngineDir)\Binaries\Win64\UE4Client.exe" Type="Executable" />
+    <BuildProduct Path="$(EngineDir)\Binaries\Win64\UE4Client.pdb" Type="SymbolFile" />
+    <BuildProduct Path="$(EngineDir)\Plugins\2D\Paper2D\Binaries\Win64\UE4Client-Paper2D-Static.lib" Type="StaticLibrary" />
+    <BuildProduct Path="$(EngineDir)\Plugins\Experimental\CharacterAI\Binaries\Win64\UE4Client-CharacterAI-Static.lib" Type="StaticLibrary" />
+    <BuildProduct Path="$(EngineDir)\Plugins\Media\WmfMedia\Binaries\Win64\UE4Client-WmfMedia-Static.lib" Type="StaticLibrary" />
+    <BuildProduct Path="$(EngineDir)\Plugins\Messaging\UdpMessaging\Binaries\Win64\UE4Client-UdpMessaging-Static.lib" Type="StaticLibrary" />
+    <BuildProduct Path="$(EngineDir)\Plugins\Runtime\CableComponent\Binaries\Win64\UE4Client-CableComponent-Static.lib" Type="StaticLibrary" />
+    <BuildProduct Path="$(EngineDir)\Plugins\Runtime\CustomMeshComponent\Binaries\Win64\UE4Client-CustomMeshComponent-Static.lib" Type="StaticLibrary" />
+    <BuildProduct Path="$(EngineDir)\Plugins\Runtime\ExampleDeviceProfileSelector\Binaries\Win64\UE4Client-ExampleDeviceProfileSelector-Static.lib" Type="StaticLibrary" />
+    <BuildProduct Path="$(EngineDir)\Plugins\Runtime\OculusRift\Binaries\Win64\UE4Client-OculusRift-Static.lib" Type="StaticLibrary" />
+    <BuildProduct Path="$(EngineDir)\Plugins\Runtime\WindowsMoviePlayer\Binaries\Win64\UE4Client-WindowsMoviePlayer-Static.lib" Type="StaticLibrary" />
+    <BuildProduct Path="$(EngineDir)\Plugins\Slate\SlateRemote\Binaries\Win64\UE4Client-SlateRemote-Static.lib" Type="StaticLibrary" />
+  </BuildProducts>
+  <RuntimeDependencies>
+    <RuntimeDependency Path="$(EngineDir)/Binaries/ThirdParty/PhysX/PhysX-3.3/Win64/VS2013/PhysX3PROFILE_x64.dll" />
+    <RuntimeDependency Path="$(EngineDir)/Binaries/ThirdParty/PhysX/PhysX-3.3/Win64/VS2013/PhysX3GpuPROFILE_x64.dll" />
+    <RuntimeDependency Path="$(EngineDir)/Binaries/ThirdParty/PhysX/PhysX-3.3/Win64/VS2013/PhysX3CommonPROFILE_x64.dll" />
+    <RuntimeDependency Path="$(EngineDir)/Binaries/ThirdParty/PhysX/PhysX-3.3/Win64/VS2013/PhysX3CookingPROFILE_x64.dll" />
+    <RuntimeDependency Path="$(EngineDir)/Binaries/ThirdParty/PhysX/PhysX-3.3/Win64/VS2013/nvToolsExt64_1.dll" />
+    <RuntimeDependency Path="$(EngineDir)/Binaries/ThirdParty/PhysX/PhysX-3.3/Win64/VS2013/PhysXDevice64.dll" />
+    <RuntimeDependency Path="$(EngineDir)/Binaries/ThirdParty/PhysX/APEX-1.3/Win64/VS2013/APEX_ClothingPROFILE_x64.dll" />
+    <RuntimeDependency Path="$(EngineDir)/Binaries/ThirdParty/PhysX/APEX-1.3/Win64/VS2013/APEX_DestructiblePROFILE_x64.dll" />
+    <RuntimeDependency Path="$(EngineDir)/Binaries/ThirdParty/PhysX/APEX-1.3/Win64/VS2013/APEX_LegacyPROFILE_x64.dll" />
+    <RuntimeDependency Path="$(EngineDir)/Binaries/ThirdParty/PhysX/APEX-1.3/Win64/VS2013/APEX_LoaderPROFILE_x64.dll" />
+    <RuntimeDependency Path="$(EngineDir)/Binaries/ThirdParty/PhysX/APEX-1.3/Win64/VS2013/APEX_ParticlesPROFILE_x64.dll" />
+    <RuntimeDependency Path="$(EngineDir)/Binaries/ThirdParty/PhysX/APEX-1.3/Win64/VS2013/ApexFrameworkPROFILE_x64.dll" />
+    <RuntimeDependency Path="$(EngineDir)/Binaries/ThirdParty/Ogg/Win64/VS2013/libogg_64.dll" />
+    <RuntimeDependency Path="$(EngineDir)/Binaries/ThirdParty/Vorbis/Win64/VS2013/libvorbis_64.dll" />
+    <RuntimeDependency Path="$(EngineDir)/Binaries/ThirdParty/Vorbis/Win64/VS2013/libvorbisfile_64.dll" />
+    <RuntimeDependency Path="$(EngineDir)/Binaries/ThirdParty/Steamworks/Steamv132/Win64/steam_api64.dll" />
+  </RuntimeDependencies>
+</BuildReceipt>
\ No newline at end of file
diff --git a/Engine/Build/Receipts/UE4Editor-Linux-Developmentx86_64-unknown-linux-gnu.target.xml b/Engine/Build/Receipts/UE4Editor-Linux-Developmentx86_64-unknown-linux-gnu.target.xml
new file mode 100644
--- /dev/null
+++ b/Engine/Build/Receipts/UE4Editor-Linux-Developmentx86_64-unknown-linux-gnu.target.xml
@@ -0,0 +1,274 @@
+<?xml version="1.0" encoding="utf-8"?>
+<BuildReceipt xmlns:xsd="http://www.w3.org/2001/XMLSchema" xmlns:xsi="http://www.w3.org/2001/XMLSchema-instance">
+  <BuildProducts>
+    <BuildProduct Path="$(EngineDir)\Binaries\Linux\UE4Editor" Type="Executable" />
+    <BuildProduct Path="$(EngineDir)\Plugins\2D\Paper2D\Binaries\Linux\libUE4Editor-Paper2D.so" Type="DynamicLibrary" />
+    <BuildProduct Path="$(EngineDir)\Plugins\2D\Paper2D\Binaries\Linux\libUE4Editor-Paper2DEditor.so" Type="DynamicLibrary" />
+    <BuildProduct Path="$(EngineDir)\Plugins\2D\Paper2D\Binaries\Linux\libUE4Editor-PaperJsonImporter.so" Type="DynamicLibrary" />
+    <BuildProduct Path="$(EngineDir)\Plugins\2D\Paper2D\Binaries\Linux\libUE4Editor-PaperTiledImporter.so" Type="DynamicLibrary" />
+    <BuildProduct Path="$(EngineDir)\Plugins\2D\Paper2D\Binaries\Linux\libUE4Editor-SmartSnapping.so" Type="DynamicLibrary" />
+    <BuildProduct Path="$(EngineDir)\Plugins\Developer\BlankPlugin\Binaries\Linux\libUE4Editor-BlankPlugin.so" Type="DynamicLibrary" />
+    <BuildProduct Path="$(EngineDir)\Plugins\Developer\GitSourceControl\Binaries\Linux\libUE4Editor-GitSourceControl.so" Type="DynamicLibrary" />
+    <BuildProduct Path="$(EngineDir)\Plugins\Developer\KDevelopSourceCodeAccess\Binaries\Linux\libUE4Editor-KDevelopSourceCodeAccess.so" Type="DynamicLibrary" />
+    <BuildProduct Path="$(EngineDir)\Plugins\Developer\PerforceSourceControl\Binaries\Linux\libUE4Editor-PerforceSourceControl.so" Type="DynamicLibrary" />
+    <BuildProduct Path="$(EngineDir)\Plugins\Developer\SubversionSourceControl\Binaries\Linux\libUE4Editor-SubversionSourceControl.so" Type="DynamicLibrary" />
+    <BuildProduct Path="$(EngineDir)\Plugins\Developer\UObjectPlugin\Binaries\Linux\libUE4Editor-UObjectPlugin.so" Type="DynamicLibrary" />
+    <BuildProduct Path="$(EngineDir)\Plugins\Editor\EpicSurvey\Binaries\Linux\libUE4Editor-EpicSurvey.so" Type="DynamicLibrary" />
+    <BuildProduct Path="$(EngineDir)\Plugins\Editor\PluginsEditor\Binaries\Linux\libUE4Editor-PluginsEditor.so" Type="DynamicLibrary" />
+    <BuildProduct Path="$(EngineDir)\Plugins\Editor\SpeedTreeImporter\Binaries\Linux\libUE4Editor-SpeedTreeImporter.so" Type="DynamicLibrary" />
+    <BuildProduct Path="$(EngineDir)\Plugins\Experimental\BlueprintStats\Binaries\Linux\libUE4Editor-BlueprintStats.so" Type="DynamicLibrary" />
+    <BuildProduct Path="$(EngineDir)\Plugins\Experimental\CharacterAI\Binaries\Linux\libUE4Editor-CharacterAI.so" Type="DynamicLibrary" />
+    <BuildProduct Path="$(EngineDir)\Plugins\Experimental\CodeEditor\Binaries\Linux\libUE4Editor-CodeEditor.so" Type="DynamicLibrary" />
+    <BuildProduct Path="$(EngineDir)\Plugins\Experimental\CodeView\Binaries\Linux\libUE4Editor-CodeView.so" Type="DynamicLibrary" />
+    <BuildProduct Path="$(EngineDir)\Plugins\Experimental\LiveEditor\Binaries\Linux\libUE4Editor-LiveEditor.so" Type="DynamicLibrary" />
+    <BuildProduct Path="$(EngineDir)\Plugins\Experimental\Phya\Binaries\Linux\libUE4Editor-Phya.so" Type="DynamicLibrary" />
+    <BuildProduct Path="$(EngineDir)\Plugins\Experimental\SimpleHMD\Binaries\Linux\libUE4Editor-SimpleHMD.so" Type="DynamicLibrary" />
+    <BuildProduct Path="$(EngineDir)\Plugins\Experimental\StructBox\Binaries\Linux\libUE4Editor-StructBox.so" Type="DynamicLibrary" />
+    <BuildProduct Path="$(EngineDir)\Plugins\Media\MediaPlayerEditor\Binaries\Linux\libUE4Editor-MediaPlayerEditor.so" Type="DynamicLibrary" />
+    <BuildProduct Path="$(EngineDir)\Plugins\Messaging\MessagingDebugger\Binaries\Linux\libUE4Editor-MessagingDebugger.so" Type="DynamicLibrary" />
+    <BuildProduct Path="$(EngineDir)\Plugins\Messaging\UdpMessaging\Binaries\Linux\libUE4Editor-UdpMessaging.so" Type="DynamicLibrary" />
+    <BuildProduct Path="$(EngineDir)\Plugins\Runtime\Analytics\AnalyticsBlueprintLibrary\Binaries\Linux\libUE4Editor-AnalyticsBlueprintLibrary.so" Type="DynamicLibrary" />
+    <BuildProduct Path="$(EngineDir)\Plugins\Runtime\Analytics\FileLogging\Binaries\Linux\libUE4Editor-FileLogging.so" Type="DynamicLibrary" />
+    <BuildProduct Path="$(EngineDir)\Plugins\Runtime\CableComponent\Binaries\Linux\libUE4Editor-CableComponent.so" Type="DynamicLibrary" />
+    <BuildProduct Path="$(EngineDir)\Plugins\Runtime\CustomMeshComponent\Binaries\Linux\libUE4Editor-CustomMeshComponent.so" Type="DynamicLibrary" />
+    <BuildProduct Path="$(EngineDir)\Plugins\Runtime\ExampleDeviceProfileSelector\Binaries\Linux\libUE4Editor-ExampleDeviceProfileSelector.so" Type="DynamicLibrary" />
+    <BuildProduct Path="$(EngineDir)\Plugins\Runtime\SoundMod\Binaries\Linux\libUE4Editor-SoundModImporter.so" Type="DynamicLibrary" />
+    <BuildProduct Path="$(EngineDir)\Plugins\Runtime\SoundVisualizations\Binaries\Linux\libUE4Editor-SoundVisualizations.so" Type="DynamicLibrary" />
+    <BuildProduct Path="$(EngineDir)\Plugins\Runtime\Steam\SteamController\Binaries\Linux\libUE4Editor-SteamController.so" Type="DynamicLibrary" />
+    <BuildProduct Path="$(EngineDir)\Plugins\Slate\SlateRemote\Binaries\Linux\libUE4Editor-SlateRemote.so" Type="DynamicLibrary" />
+    <BuildProduct Path="$(EngineDir)\Binaries\Linux\libUE4Editor-UE4Game.so" Type="DynamicLibrary" />
+    <BuildProduct Path="$(EngineDir)\Binaries\Linux\libUE4Editor-Core.so" Type="DynamicLibrary" />
+    <BuildProduct Path="$(EngineDir)\Binaries\Linux\libUE4Editor-SourceCodeAccess.so" Type="DynamicLibrary" />
+    <BuildProduct Path="$(EngineDir)\Binaries\Linux\libUE4Editor-CoreUObject.so" Type="DynamicLibrary" />
+    <BuildProduct Path="$(EngineDir)\Binaries\Linux\libUE4Editor-Projects.so" Type="DynamicLibrary" />
+    <BuildProduct Path="$(EngineDir)\Binaries\Linux\libUE4Editor-Json.so" Type="DynamicLibrary" />
+    <BuildProduct Path="$(EngineDir)\Binaries\Linux\libUE4Editor-Slate.so" Type="DynamicLibrary" />
+    <BuildProduct Path="$(EngineDir)\Binaries\Linux\libUE4Editor-InputCore.so" Type="DynamicLibrary" />
+    <BuildProduct Path="$(EngineDir)\Binaries\Linux\libUE4Editor-SlateCore.so" Type="DynamicLibrary" />
+    <BuildProduct Path="$(EngineDir)\Binaries\Linux\libUE4Editor-Settings.so" Type="DynamicLibrary" />
+    <BuildProduct Path="$(EngineDir)\Binaries\Linux\libUE4Editor-DerivedDataCache.so" Type="DynamicLibrary" />
+    <BuildProduct Path="$(EngineDir)\Binaries\Linux\NotForLicensees\libUE4Editor-DDCUtils.so" Type="DynamicLibrary" />
+    <BuildProduct Path="$(EngineDir)\Binaries\Linux\libUE4Editor-Engine.so" Type="DynamicLibrary" />
+    <BuildProduct Path="$(EngineDir)\Binaries\Linux\libUE4Editor-AppFramework.so" Type="DynamicLibrary" />
+    <BuildProduct Path="$(EngineDir)\Binaries\Linux\libUE4Editor-SlateReflector.so" Type="DynamicLibrary" />
+    <BuildProduct Path="$(EngineDir)\Binaries\Linux\libUE4Editor-Networking.so" Type="DynamicLibrary" />
+    <BuildProduct Path="$(EngineDir)\Binaries\Linux\libUE4Editor-Sockets.so" Type="DynamicLibrary" />
+    <BuildProduct Path="$(EngineDir)\Binaries\Linux\libUE4Editor-SandboxFile.so" Type="DynamicLibrary" />
+    <BuildProduct Path="$(EngineDir)\Binaries\Linux\libUE4Editor-VectorVM.so" Type="DynamicLibrary" />
+    <BuildProduct Path="$(EngineDir)\Binaries\Linux\libUE4Editor-Landscape.so" Type="DynamicLibrary" />
+    <BuildProduct Path="$(EngineDir)\Binaries\Linux\libUE4Editor-RenderCore.so" Type="DynamicLibrary" />
+    <BuildProduct Path="$(EngineDir)\Binaries\Linux\libUE4Editor-RHI.so" Type="DynamicLibrary" />
+    <BuildProduct Path="$(EngineDir)\Binaries\Linux\libUE4Editor-NullDrv.so" Type="DynamicLibrary" />
+    <BuildProduct Path="$(EngineDir)\Binaries\Linux\libUE4Editor-OpenGLDrv.so" Type="DynamicLibrary" />
+    <BuildProduct Path="$(EngineDir)\Binaries\Linux\libUE4Editor-ShaderCore.so" Type="DynamicLibrary" />
+    <BuildProduct Path="$(EngineDir)\Binaries\Linux\libUE4Editor-UtilityShaders.so" Type="DynamicLibrary" />
+    <BuildProduct Path="$(EngineDir)\Binaries\Linux\libUE4Editor-ImageWrapper.so" Type="DynamicLibrary" />
+    <BuildProduct Path="$(EngineDir)\Binaries\Linux\libUE4Editor-Renderer.so" Type="DynamicLibrary" />
+    <BuildProduct Path="$(EngineDir)\Binaries\Linux\libUE4Editor-HeadMountedDisplay.so" Type="DynamicLibrary" />
+    <BuildProduct Path="$(EngineDir)\Binaries\Linux\libUE4Editor-Foliage.so" Type="DynamicLibrary" />
+    <BuildProduct Path="$(EngineDir)\Binaries\Linux\libUE4Editor-RawMesh.so" Type="DynamicLibrary" />
+    <BuildProduct Path="$(EngineDir)\Binaries\Linux\libUE4Editor-UnrealEd.so" Type="DynamicLibrary" />
+    <BuildProduct Path="$(EngineDir)\Binaries\Linux\libUE4Editor-AnimGraph.so" Type="DynamicLibrary" />
+    <BuildProduct Path="$(EngineDir)\Binaries\Linux\libUE4Editor-GraphEditor.so" Type="DynamicLibrary" />
+    <BuildProduct Path="$(EngineDir)\Binaries\Linux\libUE4Editor-EditorStyle.so" Type="DynamicLibrary" />
+    <BuildProduct Path="$(EngineDir)\Binaries\Linux\libUE4Editor-DesktopPlatform.so" Type="DynamicLibrary" />
+    <BuildProduct Path="$(EngineDir)\Binaries\Linux\libUE4Editor-AssetRegistry.so" Type="DynamicLibrary" />
+    <BuildProduct Path="$(EngineDir)\Binaries\Linux\libUE4Editor-DirectoryWatcher.so" Type="DynamicLibrary" />
+    <BuildProduct Path="$(EngineDir)\Binaries\Linux\libUE4Editor-ClassViewer.so" Type="DynamicLibrary" />
+    <BuildProduct Path="$(EngineDir)\Binaries\Linux\libUE4Editor-PropertyEditor.so" Type="DynamicLibrary" />
+    <BuildProduct Path="$(EngineDir)\Binaries\Linux\libUE4Editor-EditorWidgets.so" Type="DynamicLibrary" />
+    <BuildProduct Path="$(EngineDir)\Binaries\Linux\libUE4Editor-Documentation.so" Type="DynamicLibrary" />
+    <BuildProduct Path="$(EngineDir)\Binaries\Linux\libUE4Editor-Analytics.so" Type="DynamicLibrary" />
+    <BuildProduct Path="$(EngineDir)\Binaries\Linux\libUE4Editor-SourceControl.so" Type="DynamicLibrary" />
+    <BuildProduct Path="$(EngineDir)\Binaries\Linux\libUE4Editor-MessageLog.so" Type="DynamicLibrary" />
+    <BuildProduct Path="$(EngineDir)\Binaries\Linux\libUE4Editor-AssetTools.so" Type="DynamicLibrary" />
+    <BuildProduct Path="$(EngineDir)\Binaries\Linux\libUE4Editor-CurveAssetEditor.so" Type="DynamicLibrary" />
+    <BuildProduct Path="$(EngineDir)\Binaries\Linux\libUE4Editor-WorkspaceMenuStructure.so" Type="DynamicLibrary" />
+    <BuildProduct Path="$(EngineDir)\Binaries\Linux\libUE4Editor-TextureEditor.so" Type="DynamicLibrary" />
+    <BuildProduct Path="$(EngineDir)\Binaries\Linux\libUE4Editor-MainFrame.so" Type="DynamicLibrary" />
+    <BuildProduct Path="$(EngineDir)\Binaries\Linux\libUE4Editor-CrashTracker.so" Type="DynamicLibrary" />
+    <BuildProduct Path="$(EngineDir)\Binaries\Linux\libUE4Editor-EngineSettings.so" Type="DynamicLibrary" />
+    <BuildProduct Path="$(EngineDir)\Binaries\Linux\libUE4Editor-SourceControlWindows.so" Type="DynamicLibrary" />
+    <BuildProduct Path="$(EngineDir)\Binaries\Linux\libUE4Editor-TargetPlatform.so" Type="DynamicLibrary" />
+    <BuildProduct Path="$(EngineDir)\Binaries\Linux\libUE4Editor-ShaderPreprocessor.so" Type="DynamicLibrary" />
+    <BuildProduct Path="$(EngineDir)\Binaries\Linux\libUE4Editor-ShaderFormatOpenGL.so" Type="DynamicLibrary" />
+    <BuildProduct Path="$(EngineDir)\Binaries\Linux\libUE4Editor-ShaderCompilerCommon.so" Type="DynamicLibrary" />
+    <BuildProduct Path="$(EngineDir)\Binaries\Linux\libUE4Editor-TextureFormatDXT.so" Type="DynamicLibrary" />
+    <BuildProduct Path="$(EngineDir)\Binaries\Linux\libUE4Editor-ImageCore.so" Type="DynamicLibrary" />
+    <BuildProduct Path="$(EngineDir)\Binaries\Linux\libUE4Editor-TextureFormatPVR.so" Type="DynamicLibrary" />
+    <BuildProduct Path="$(EngineDir)\Binaries\Linux\libUE4Editor-TextureFormatASTC.so" Type="DynamicLibrary" />
+    <BuildProduct Path="$(EngineDir)\Binaries\Linux\libUE4Editor-TextureFormatUncompressed.so" Type="DynamicLibrary" />
+    <BuildProduct Path="$(EngineDir)\Binaries\Linux\libUE4Editor-AudioFormatOgg.so" Type="DynamicLibrary" />
+    <BuildProduct Path="$(EngineDir)\Binaries\Linux\libUE4Editor-AudioFormatOpus.so" Type="DynamicLibrary" />
+    <BuildProduct Path="$(EngineDir)\Binaries\Linux\libUE4Editor-PhysXFormats.so" Type="DynamicLibrary" />
+    <BuildProduct Path="$(EngineDir)\Binaries\Linux\libUE4Editor-TranslationEditor.so" Type="DynamicLibrary" />
+    <BuildProduct Path="$(EngineDir)\Binaries\Linux\libUE4Editor-Internationalization.so" Type="DynamicLibrary" />
+    <BuildProduct Path="$(EngineDir)\Binaries\Linux\libUE4Editor-DeviceProfileEditor.so" Type="DynamicLibrary" />
+    <BuildProduct Path="$(EngineDir)\Binaries\Linux\libUE4Editor-LevelEditor.so" Type="DynamicLibrary" />
+    <BuildProduct Path="$(EngineDir)\Binaries\Linux\libUE4Editor-NewsFeed.so" Type="DynamicLibrary" />
+    <BuildProduct Path="$(EngineDir)\Binaries\Linux\libUE4Editor-HTTP.so" Type="DynamicLibrary" />
+    <BuildProduct Path="$(EngineDir)\Binaries\Linux\libUE4Editor-OnlineSubsystem.so" Type="DynamicLibrary" />
+    <BuildProduct Path="$(EngineDir)\Binaries\Linux\libUE4Editor-StatsViewer.so" Type="DynamicLibrary" />
+    <BuildProduct Path="$(EngineDir)\Binaries\Linux\libUE4Editor-DeviceProfileServices.so" Type="DynamicLibrary" />
+    <BuildProduct Path="$(EngineDir)\Binaries\Linux\libUE4Editor-ContentBrowser.so" Type="DynamicLibrary" />
+    <BuildProduct Path="$(EngineDir)\Binaries\Linux\libUE4Editor-AddContentDialog.so" Type="DynamicLibrary" />
+    <BuildProduct Path="$(EngineDir)\Binaries\Linux\libUE4Editor-PakFile.so" Type="DynamicLibrary" />
+    <BuildProduct Path="$(EngineDir)\Binaries\Linux\libUE4Editor-WidgetCarousel.so" Type="DynamicLibrary" />
+    <BuildProduct Path="$(EngineDir)\Binaries\Linux\libUE4Editor-CollectionManager.so" Type="DynamicLibrary" />
+    <BuildProduct Path="$(EngineDir)\Binaries\Linux\libUE4Editor-GameProjectGeneration.so" Type="DynamicLibrary" />
+    <BuildProduct Path="$(EngineDir)\Binaries\Linux\libUE4Editor-HardwareTargeting.so" Type="DynamicLibrary" />
+    <BuildProduct Path="$(EngineDir)\Binaries\Linux\libUE4Editor-ReferenceViewer.so" Type="DynamicLibrary" />
+    <BuildProduct Path="$(EngineDir)\Binaries\Linux\libUE4Editor-SceneOutliner.so" Type="DynamicLibrary" />
+    <BuildProduct Path="$(EngineDir)\Binaries\Linux\libUE4Editor-ActorPickerMode.so" Type="DynamicLibrary" />
+    <BuildProduct Path="$(EngineDir)\Binaries\Linux\libUE4Editor-WebBrowser.so" Type="DynamicLibrary" />
+    <BuildProduct Path="$(EngineDir)\Binaries\Linux\libUE4Editor-CEF3Utils.so" Type="DynamicLibrary" />
+    <BuildProduct Path="$(EngineDir)\Binaries\Linux\libUE4Editor-Persona.so" Type="DynamicLibrary" />
+    <BuildProduct Path="$(EngineDir)\Binaries\Linux\libUE4Editor-Kismet.so" Type="DynamicLibrary" />
+    <BuildProduct Path="$(EngineDir)\Binaries\Linux\libUE4Editor-Merge.so" Type="DynamicLibrary" />
+    <BuildProduct Path="$(EngineDir)\Binaries\Linux\libUE4Editor-KismetWidgets.so" Type="DynamicLibrary" />
+    <BuildProduct Path="$(EngineDir)\Binaries\Linux\libUE4Editor-BlueprintGraph.so" Type="DynamicLibrary" />
+    <BuildProduct Path="$(EngineDir)\Binaries\Linux\libUE4Editor-KismetCompiler.so" Type="DynamicLibrary" />
+    <BuildProduct Path="$(EngineDir)\Binaries\Linux\libUE4Editor-MovieSceneCore.so" Type="DynamicLibrary" />
+    <BuildProduct Path="$(EngineDir)\Binaries\Linux\libUE4Editor-MovieSceneTools.so" Type="DynamicLibrary" />
+    <BuildProduct Path="$(EngineDir)\Binaries\Linux\libUE4Editor-MovieSceneCoreTypes.so" Type="DynamicLibrary" />
+    <BuildProduct Path="$(EngineDir)\Binaries\Linux\libUE4Editor-Sequencer.so" Type="DynamicLibrary" />
+    <BuildProduct Path="$(EngineDir)\Binaries\Linux\libUE4Editor-SequencerWidgets.so" Type="DynamicLibrary" />
+    <BuildProduct Path="$(EngineDir)\Binaries\Linux\libUE4Editor-FunctionalTesting.so" Type="DynamicLibrary" />
+    <BuildProduct Path="$(EngineDir)\Binaries\Linux\libUE4Editor-AIModule.so" Type="DynamicLibrary" />
+    <BuildProduct Path="$(EngineDir)\Binaries\Linux\libUE4Editor-AITestSuite.so" Type="DynamicLibrary" />
+    <BuildProduct Path="$(EngineDir)\Binaries\Linux\libUE4Editor-Navmesh.so" Type="DynamicLibrary" />
+    <BuildProduct Path="$(EngineDir)\Binaries\Linux\libUE4Editor-GameplayTags.so" Type="DynamicLibrary" />
+    <BuildProduct Path="$(EngineDir)\Binaries\Linux\libUE4Editor-JsonUtilities.so" Type="DynamicLibrary" />
+    <BuildProduct Path="$(EngineDir)\Binaries\Linux\libUE4Editor-HotReload.so" Type="DynamicLibrary" />
+    <BuildProduct Path="$(EngineDir)\Binaries\Linux\libUE4Editor-BlueprintCompilerCppBackend.so" Type="DynamicLibrary" />
+    <BuildProduct Path="$(EngineDir)\Binaries\Linux\libUE4Editor-DeviceManager.so" Type="DynamicLibrary" />
+    <BuildProduct Path="$(EngineDir)\Binaries\Linux\libUE4Editor-SettingsEditor.so" Type="DynamicLibrary" />
+    <BuildProduct Path="$(EngineDir)\Binaries\Linux\libUE4Editor-SharedSettingsWidgets.so" Type="DynamicLibrary" />
+    <BuildProduct Path="$(EngineDir)\Binaries\Linux\libUE4Editor-ExternalImagePicker.so" Type="DynamicLibrary" />
+    <BuildProduct Path="$(EngineDir)\Binaries\Linux\libUE4Editor-SessionFrontend.so" Type="DynamicLibrary" />
+    <BuildProduct Path="$(EngineDir)\Binaries\Linux\libUE4Editor-SessionServices.so" Type="DynamicLibrary" />
+    <BuildProduct Path="$(EngineDir)\Binaries\Linux\libUE4Editor-EngineMessages.so" Type="DynamicLibrary" />
+    <BuildProduct Path="$(EngineDir)\Binaries\Linux\libUE4Editor-SessionMessages.so" Type="DynamicLibrary" />
+    <BuildProduct Path="$(EngineDir)\Binaries\Linux\libUE4Editor-AutomationWindow.so" Type="DynamicLibrary" />
+    <BuildProduct Path="$(EngineDir)\Binaries\Linux\libUE4Editor-AutomationController.so" Type="DynamicLibrary" />
+    <BuildProduct Path="$(EngineDir)\Binaries\Linux\libUE4Editor-AutomationMessages.so" Type="DynamicLibrary" />
+    <BuildProduct Path="$(EngineDir)\Binaries\Linux\libUE4Editor-UnrealEdMessages.so" Type="DynamicLibrary" />
+    <BuildProduct Path="$(EngineDir)\Binaries\Linux\libUE4Editor-ScreenShotComparison.so" Type="DynamicLibrary" />
+    <BuildProduct Path="$(EngineDir)\Binaries\Linux\libUE4Editor-ScreenShotComparisonTools.so" Type="DynamicLibrary" />
+    <BuildProduct Path="$(EngineDir)\Binaries\Linux\libUE4Editor-Profiler.so" Type="DynamicLibrary" />
+    <BuildProduct Path="$(EngineDir)\Binaries\Linux\libUE4Editor-ProfilerClient.so" Type="DynamicLibrary" />
+    <BuildProduct Path="$(EngineDir)\Binaries\Linux\libUE4Editor-ProfilerMessages.so" Type="DynamicLibrary" />
+    <BuildProduct Path="$(EngineDir)\Binaries\Linux\libUE4Editor-Layers.so" Type="DynamicLibrary" />
+    <BuildProduct Path="$(EngineDir)\Binaries\Linux\libUE4Editor-Levels.so" Type="DynamicLibrary" />
+    <BuildProduct Path="$(EngineDir)\Binaries\Linux\libUE4Editor-WorldBrowser.so" Type="DynamicLibrary" />
+    <BuildProduct Path="$(EngineDir)\Binaries\Linux\libUE4Editor-MeshUtilities.so" Type="DynamicLibrary" />
+    <BuildProduct Path="$(EngineDir)\Binaries\Linux\libUE4Editor-MeshSimplifier.so" Type="DynamicLibrary" />
+    <BuildProduct Path="$(EngineDir)\Binaries\Linux\libUE4Editor-LandscapeEditor.so" Type="DynamicLibrary" />
+    <BuildProduct Path="$(EngineDir)\Binaries\Linux\libUE4Editor-TaskBrowser.so" Type="DynamicLibrary" />
+    <BuildProduct Path="$(EngineDir)\Binaries\Linux\libUE4Editor-NewLevelDialog.so" Type="DynamicLibrary" />
+    <BuildProduct Path="$(EngineDir)\Binaries\Linux\libUE4Editor-PlacementMode.so" Type="DynamicLibrary" />
+    <BuildProduct Path="$(EngineDir)\Binaries\Linux\libUE4Editor-BspMode.so" Type="DynamicLibrary" />
+    <BuildProduct Path="$(EngineDir)\Binaries\Linux\libUE4Editor-UserFeedback.so" Type="DynamicLibrary" />
+    <BuildProduct Path="$(EngineDir)\Binaries\Linux\libUE4Editor-IntroTutorials.so" Type="DynamicLibrary" />
+    <BuildProduct Path="$(EngineDir)\Binaries\Linux\libUE4Editor-LauncherServices.so" Type="DynamicLibrary" />
+    <BuildProduct Path="$(EngineDir)\Binaries\Linux\libUE4Editor-UndoHistory.so" Type="DynamicLibrary" />
+    <BuildProduct Path="$(EngineDir)\Binaries\Linux\libUE4Editor-ProjectTargetPlatformEditor.so" Type="DynamicLibrary" />
+    <BuildProduct Path="$(EngineDir)\Binaries\Linux\libUE4Editor-OutputLog.so" Type="DynamicLibrary" />
+    <BuildProduct Path="$(EngineDir)\Binaries\Linux\libUE4Editor-SuperSearch.so" Type="DynamicLibrary" />
+    <BuildProduct Path="$(EngineDir)\Binaries\Linux\libUE4Editor-EditorLiveStreaming.so" Type="DynamicLibrary" />
+    <BuildProduct Path="$(EngineDir)\Binaries\Linux\libUE4Editor-GameLiveStreaming.so" Type="DynamicLibrary" />
+    <BuildProduct Path="$(EngineDir)\Binaries\Linux\libUE4Editor-OnlineSubsystemUtils.so" Type="DynamicLibrary" />
+    <BuildProduct Path="$(EngineDir)\Binaries\Linux\libUE4Editor-Voice.so" Type="DynamicLibrary" />
+    <BuildProduct Path="$(EngineDir)\Binaries\Linux\libUE4Editor-LocalizationDashboard.so" Type="DynamicLibrary" />
+    <BuildProduct Path="$(EngineDir)\Binaries\Linux\libUE4Editor-CurveTableEditor.so" Type="DynamicLibrary" />
+    <BuildProduct Path="$(EngineDir)\Binaries\Linux\libUE4Editor-DataTableEditor.so" Type="DynamicLibrary" />
+    <BuildProduct Path="$(EngineDir)\Binaries\Linux\libUE4Editor-DestructibleMeshEditor.so" Type="DynamicLibrary" />
+    <BuildProduct Path="$(EngineDir)\Binaries\Linux\libUE4Editor-MaterialEditor.so" Type="DynamicLibrary" />
+    <BuildProduct Path="$(EngineDir)\Binaries\Linux\libUE4Editor-FontEditor.so" Type="DynamicLibrary" />
+    <BuildProduct Path="$(EngineDir)\Binaries\Linux\libUE4Editor-DesktopWidgets.so" Type="DynamicLibrary" />
+    <BuildProduct Path="$(EngineDir)\Binaries\Linux\libUE4Editor-SoundCueEditor.so" Type="DynamicLibrary" />
+    <BuildProduct Path="$(EngineDir)\Binaries\Linux\libUE4Editor-SoundClassEditor.so" Type="DynamicLibrary" />
+    <BuildProduct Path="$(EngineDir)\Binaries\Linux\libUE4Editor-OnlineBlueprintSupport.so" Type="DynamicLibrary" />
+    <BuildProduct Path="$(EngineDir)\Binaries\Linux\libUE4Editor-InputBindingEditor.so" Type="DynamicLibrary" />
+    <BuildProduct Path="$(EngineDir)\Binaries\Linux\libUE4Editor-LauncherAutomatedService.so" Type="DynamicLibrary" />
+    <BuildProduct Path="$(EngineDir)\Binaries\Linux\libUE4Editor-NetworkFileSystem.so" Type="DynamicLibrary" />
+    <BuildProduct Path="$(EngineDir)\Binaries\Linux\libUE4Editor-SwarmInterface.so" Type="DynamicLibrary" />
+    <BuildProduct Path="$(EngineDir)\Binaries\Linux\libUE4Editor-Messaging.so" Type="DynamicLibrary" />
+    <BuildProduct Path="$(EngineDir)\Binaries\Linux\libUE4Editor-TargetDeviceServices.so" Type="DynamicLibrary" />
+    <BuildProduct Path="$(EngineDir)\Binaries\Linux\libUE4Editor-InternationalizationSettings.so" Type="DynamicLibrary" />
+    <BuildProduct Path="$(EngineDir)\Binaries\Linux\libUE4Editor-MeshPaint.so" Type="DynamicLibrary" />
+    <BuildProduct Path="$(EngineDir)\Binaries\Linux\libUE4Editor-UnrealAudio.so" Type="DynamicLibrary" />
+    <BuildProduct Path="$(EngineDir)\Binaries\Linux\libUE4Editor-StaticMeshEditor.so" Type="DynamicLibrary" />
+    <BuildProduct Path="$(EngineDir)\Binaries\Linux\libUE4Editor-Cascade.so" Type="DynamicLibrary" />
+    <BuildProduct Path="$(EngineDir)\Binaries\Linux\libUE4Editor-DistCurveEditor.so" Type="DynamicLibrary" />
+    <BuildProduct Path="$(EngineDir)\Binaries\Linux\libUE4Editor-UMGEditor.so" Type="DynamicLibrary" />
+    <BuildProduct Path="$(EngineDir)\Binaries\Linux\libUE4Editor-UMG.so" Type="DynamicLibrary" />
+    <BuildProduct Path="$(EngineDir)\Binaries\Linux\libUE4Editor-SlateRHIRenderer.so" Type="DynamicLibrary" />
+    <BuildProduct Path="$(EngineDir)\Binaries\Linux\libUE4Editor-DetailCustomizations.so" Type="DynamicLibrary" />
+    <BuildProduct Path="$(EngineDir)\Binaries\Linux\libUE4Editor-MoviePlayer.so" Type="DynamicLibrary" />
+    <BuildProduct Path="$(EngineDir)\Binaries\Linux\libUE4Editor-MediaAssets.so" Type="DynamicLibrary" />
+    <BuildProduct Path="$(EngineDir)\Binaries\Linux\libUE4Editor-Media.so" Type="DynamicLibrary" />
+    <BuildProduct Path="$(EngineDir)\Binaries\Linux\libUE4Editor-Matinee.so" Type="DynamicLibrary" />
+    <BuildProduct Path="$(EngineDir)\Binaries\Linux\libUE4Editor-EditorSettingsViewer.so" Type="DynamicLibrary" />
+    <BuildProduct Path="$(EngineDir)\Binaries\Linux\libUE4Editor-ComponentVisualizers.so" Type="DynamicLibrary" />
+    <BuildProduct Path="$(EngineDir)\Binaries\Linux\libUE4Editor-PackagesDialog.so" Type="DynamicLibrary" />
+    <BuildProduct Path="$(EngineDir)\Binaries\Linux\libUE4Editor-PhAT.so" Type="DynamicLibrary" />
+    <BuildProduct Path="$(EngineDir)\Binaries\Linux\libUE4Editor-ProjectLauncher.so" Type="DynamicLibrary" />
+    <BuildProduct Path="$(EngineDir)\Binaries\Linux\libUE4Editor-GeometryMode.so" Type="DynamicLibrary" />
+    <BuildProduct Path="$(EngineDir)\Binaries\Linux\libUE4Editor-TextureAlignMode.so" Type="DynamicLibrary" />
+    <BuildProduct Path="$(EngineDir)\Binaries\Linux\libUE4Editor-FoliageEdit.so" Type="DynamicLibrary" />
+    <BuildProduct Path="$(EngineDir)\Binaries\Linux\libUE4Editor-PackageDependencyInfo.so" Type="DynamicLibrary" />
+    <BuildProduct Path="$(EngineDir)\Binaries\Linux\libUE4Editor-Blutility.so" Type="DynamicLibrary" />
+    <BuildProduct Path="$(EngineDir)\Binaries\Linux\libUE4Editor-NiagaraEditor.so" Type="DynamicLibrary" />
+    <BuildProduct Path="$(EngineDir)\Binaries\Linux\libUE4Editor-ProjectSettingsViewer.so" Type="DynamicLibrary" />
+    <BuildProduct Path="$(EngineDir)\Binaries\Linux\libUE4Editor-PListEditor.so" Type="DynamicLibrary" />
+    <BuildProduct Path="$(EngineDir)\Binaries\Linux\libUE4Editor-XmlParser.so" Type="DynamicLibrary" />
+    <BuildProduct Path="$(EngineDir)\Binaries\Linux\libUE4Editor-BehaviorTreeEditor.so" Type="DynamicLibrary" />
+    <BuildProduct Path="$(EngineDir)\Binaries\Linux\libUE4Editor-AIGraph.so" Type="DynamicLibrary" />
+    <BuildProduct Path="$(EngineDir)\Binaries\Linux\libUE4Editor-GameplayDebugger.so" Type="DynamicLibrary" />
+    <BuildProduct Path="$(EngineDir)\Binaries\Linux\libUE4Editor-GameplayAbilities.so" Type="DynamicLibrary" />
+    <BuildProduct Path="$(EngineDir)\Binaries\Linux\libUE4Editor-GameplayTagsEditor.so" Type="DynamicLibrary" />
+    <BuildProduct Path="$(EngineDir)\Binaries\Linux\libUE4Editor-EnvironmentQueryEditor.so" Type="DynamicLibrary" />
+    <BuildProduct Path="$(EngineDir)\Binaries\Linux\libUE4Editor-ViewportSnapping.so" Type="DynamicLibrary" />
+    <BuildProduct Path="$(EngineDir)\Binaries\Linux\libUE4Editor-GameplayAbilitiesEditor.so" Type="DynamicLibrary" />
+    <BuildProduct Path="$(EngineDir)\Binaries\Linux\IOS\libUE4Editor-IOSPlatformEditor.so" Type="DynamicLibrary" />
+    <BuildProduct Path="$(EngineDir)\Binaries\Linux\IOS\libUE4Editor-IOSRuntimeSettings.so" Type="DynamicLibrary" />
+    <BuildProduct Path="$(EngineDir)\Binaries\Linux\HTML5\libUE4Editor-HTML5PlatformEditor.so" Type="DynamicLibrary" />
+    <BuildProduct Path="$(EngineDir)\Binaries\Linux\HTML5\libUE4Editor-HTML5TargetPlatform.so" Type="DynamicLibrary" />
+    <BuildProduct Path="$(EngineDir)\Binaries\Linux\libUE4Editor-CollisionAnalyzer.so" Type="DynamicLibrary" />
+    <BuildProduct Path="$(EngineDir)\Binaries\Linux\libUE4Editor-LogVisualizer.so" Type="DynamicLibrary" />
+    <BuildProduct Path="$(EngineDir)\Binaries\Linux\libUE4Editor-TaskGraph.so" Type="DynamicLibrary" />
+    <BuildProduct Path="$(EngineDir)\Binaries\Linux\libUE4Editor-SynthBenchmark.so" Type="DynamicLibrary" />
+    <BuildProduct Path="$(EngineDir)\Binaries\Linux\libUE4Editor-DatabaseSupport.so" Type="DynamicLibrary" />
+    <BuildProduct Path="$(EngineDir)\Binaries\Linux\libUE4Editor-StreamingPauseRendering.so" Type="DynamicLibrary" />
+    <BuildProduct Path="$(EngineDir)\Binaries\Linux\libUE4Editor-BlankModule.so" Type="DynamicLibrary" />
+    <BuildProduct Path="$(EngineDir)\Binaries\Linux\Linux\libUE4Editor-LinuxTargetPlatform.so" Type="DynamicLibrary" />
+    <BuildProduct Path="$(EngineDir)\Binaries\Linux\Linux\libUE4Editor-LinuxNoEditorTargetPlatform.so" Type="DynamicLibrary" />
+    <BuildProduct Path="$(EngineDir)\Binaries\Linux\Linux\libUE4Editor-LinuxServerTargetPlatform.so" Type="DynamicLibrary" />
+    <BuildProduct Path="$(EngineDir)\Binaries\Linux\libUE4Editor-DesktopTargetPlatform.so" Type="DynamicLibrary" />
+    <BuildProduct Path="$(EngineDir)\Binaries\Linux\libUE4Editor-LaunchDaemonMessages.so" Type="DynamicLibrary" />
+    <BuildProduct Path="$(EngineDir)\Binaries\Linux\libUE4Editor-AnalyticsET.so" Type="DynamicLibrary" />
+    <BuildProduct Path="$(EngineDir)\Binaries\Linux\libUE4Editor-NetworkReplayStreaming.so" Type="DynamicLibrary" />
+    <BuildProduct Path="$(EngineDir)\Binaries\Linux\libUE4Editor-NullNetworkReplayStreaming.so" Type="DynamicLibrary" />
+    <BuildProduct Path="$(EngineDir)\Binaries\Linux\libUE4Editor-HttpNetworkReplayStreaming.so" Type="DynamicLibrary" />
+    <BuildProduct Path="$(EngineDir)\Binaries\Linux\libUE4Editor-Advertising.so" Type="DynamicLibrary" />
+    <BuildProduct Path="$(EngineDir)\Binaries\Linux\libUE4Editor-Serialization.so" Type="DynamicLibrary" />
+    <BuildProduct Path="$(EngineDir)\Binaries\Linux\libUE4Editor-NetworkFile.so" Type="DynamicLibrary" />
+    <BuildProduct Path="$(EngineDir)\Binaries\Linux\libUE4Editor-StreamingFile.so" Type="DynamicLibrary" />
+    <BuildProduct Path="$(EngineDir)\Binaries\Linux\libUE4Editor-AutomationWorker.so" Type="DynamicLibrary" />
+    <BuildProduct Path="$(EngineDir)\Binaries\Linux\libUE4Editor-ALAudio.so" Type="DynamicLibrary" />
+    <BuildProduct Path="$(EngineDir)\Binaries\Linux\libUE4Editor-RealtimeProfiler.so" Type="DynamicLibrary" />
+    <BuildProduct Path="$(EngineDir)\Binaries\Linux\libUE4Editor-ProfilerService.so" Type="DynamicLibrary" />
+    <BuildProduct Path="$(EngineDir)\Binaries\Linux\libUE4Editor-Toolbox.so" Type="DynamicLibrary" />
+    <BuildProduct Path="$(EngineDir)\Binaries\Linux\libUE4Editor-GammaUI.so" Type="DynamicLibrary" />
+    <BuildProduct Path="$(EngineDir)\Binaries\Linux\libUE4Editor-ModuleUI.so" Type="DynamicLibrary" />
+    <BuildProduct Path="$(EngineDir)\Binaries\Linux\libUE4Editor-TextureCompressor.so" Type="DynamicLibrary" />
+    <BuildProduct Path="$(EngineDir)\Binaries\Linux\NotForLicensees\libUE4Editor-OnlineSubsystemMcp.so" Type="DynamicLibrary" />
+    <BuildProduct Path="$(EngineDir)\Binaries\Linux\NotForLicensees\libUE4Editor-XMPP.so" Type="DynamicLibrary" />
+    <BuildProduct Path="$(EngineDir)\Plugins\Runtime\LiveEditorListenServer\Binaries\Linux\libUE4Editor-LiveEditorListenServer.so" Type="DynamicLibrary" />
+    <BuildProduct Path="$(EngineDir)\Plugins\Runtime\SoundMod\Binaries\Linux\libUE4Editor-SoundMod.so" Type="DynamicLibrary" />
+    <BuildProduct Path="$(EngineDir)\Binaries\Linux\libUE4Editor-InputDevice.so" Type="DynamicLibrary" />
+    <BuildProduct Path="$(EngineDir)\Binaries\Linux\libUE4Editor-OnlineSubsystemNull.so" Type="DynamicLibrary" />
+  </BuildProducts>
+  <RuntimeDependencies />
+</BuildReceipt>
\ No newline at end of file
diff --git a/Engine/Build/Receipts/UE4Game-Android-Development.target.xml b/Engine/Build/Receipts/UE4Game-Android-Development.target.xml
new file mode 100644
--- /dev/null
+++ b/Engine/Build/Receipts/UE4Game-Android-Development.target.xml
@@ -0,0 +1,101 @@
+<?xml version="1.0" encoding="utf-8"?>
+<BuildReceipt xmlns:xsd="http://www.w3.org/2001/XMLSchema" xmlns:xsi="http://www.w3.org/2001/XMLSchema-instance">
+  <BuildProducts>
+    <BuildProduct Path="$(EngineDir)\Binaries\Android\UE4Game-armv7-es2.so" Type="Executable" />
+    <BuildProduct Path="$(EngineDir)\Binaries\Android\UE4Game-armv7-es2.apk" Type="Executable" />
+    <BuildProduct Path="$(EngineDir)\Plugins\2D\Paper2D\Binaries\Android\UE4Game-Paper2D-Static-armv7-es2.a" Type="StaticLibrary" />
+    <BuildProduct Path="$(EngineDir)\Plugins\Experimental\CharacterAI\Binaries\Android\UE4Game-CharacterAI-Static-armv7-es2.a" Type="StaticLibrary" />
+    <BuildProduct Path="$(EngineDir)\Plugins\Experimental\SimpleHMD\Binaries\Android\UE4Game-SimpleHMD-Static-armv7-es2.a" Type="StaticLibrary" />
+    <BuildProduct Path="$(EngineDir)\Plugins\Experimental\StructBox\Binaries\Android\UE4Game-StructBox-Static-armv7-es2.a" Type="StaticLibrary" />
+    <BuildProduct Path="$(EngineDir)\Plugins\Media\AndroidMedia\Binaries\Android\UE4Game-AndroidMedia-Static-armv7-es2.a" Type="StaticLibrary" />
+    <BuildProduct Path="$(EngineDir)\Plugins\Messaging\UdpMessaging\Binaries\Android\UE4Game-UdpMessaging-Static-armv7-es2.a" Type="StaticLibrary" />
+    <BuildProduct Path="$(EngineDir)\Plugins\Runtime\Analytics\AnalyticsBlueprintLibrary\Binaries\Android\UE4Game-AnalyticsBlueprintLibrary-Static-armv7-es2.a" Type="StaticLibrary" />
+    <BuildProduct Path="$(EngineDir)\Plugins\Runtime\Analytics\FileLogging\Binaries\Android\UE4Game-FileLogging-Static-armv7-es2.a" Type="StaticLibrary" />
+    <BuildProduct Path="$(EngineDir)\Plugins\Runtime\AndroidDeviceProfileSelector\Binaries\Android\UE4Game-AndroidDeviceProfileSelector-Static-armv7-es2.a" Type="StaticLibrary" />
+    <BuildProduct Path="$(EngineDir)\Plugins\Runtime\AndroidMoviePlayer\Binaries\Android\UE4Game-AndroidMoviePlayer-Static-armv7-es2.a" Type="StaticLibrary" />
+    <BuildProduct Path="$(EngineDir)\Plugins\Runtime\CableComponent\Binaries\Android\UE4Game-CableComponent-Static-armv7-es2.a" Type="StaticLibrary" />
+    <BuildProduct Path="$(EngineDir)\Plugins\Runtime\CustomMeshComponent\Binaries\Android\UE4Game-CustomMeshComponent-Static-armv7-es2.a" Type="StaticLibrary" />
+    <BuildProduct Path="$(EngineDir)\Plugins\Runtime\ExampleDeviceProfileSelector\Binaries\Android\UE4Game-ExampleDeviceProfileSelector-Static-armv7-es2.a" Type="StaticLibrary" />
+    <BuildProduct Path="$(EngineDir)\Plugins\Runtime\GearVR\Binaries\Android\UE4Game-GearVR-Static-armv7-es2.a" Type="StaticLibrary" />
+    <BuildProduct Path="$(EngineDir)\Plugins\Runtime\SoundMod\Binaries\Android\UE4Game-SoundMod-Static-armv7-es2.a" Type="StaticLibrary" />
+    <BuildProduct Path="$(EngineDir)\Plugins\Runtime\SoundVisualizations\Binaries\Android\UE4Game-SoundVisualizations-Static-armv7-es2.a" Type="StaticLibrary" />
+    <BuildProduct Path="$(EngineDir)\Plugins\ScriptPlugin\Binaries\Android\UE4Game-ScriptPlugin-Static-armv7-es2.a" Type="StaticLibrary" />
+    <BuildProduct Path="$(EngineDir)\Plugins\Slate\SlateRemote\Binaries\Android\UE4Game-SlateRemote-Static-armv7-es2.a" Type="StaticLibrary" />
+    <BuildProduct Path="$(EngineDir)\Binaries\Android\UE4Game-Launch-armv7-es2.a" Type="StaticLibrary" />
+    <BuildProduct Path="$(EngineDir)\Binaries\Android\UE4Game-Core-armv7-es2.a" Type="StaticLibrary" />
+    <BuildProduct Path="$(EngineDir)\Binaries\Android\UE4Game-CoreUObject-armv7-es2.a" Type="StaticLibrary" />
+    <BuildProduct Path="$(EngineDir)\Binaries\Android\UE4Game-Projects-armv7-es2.a" Type="StaticLibrary" />
+    <BuildProduct Path="$(EngineDir)\Binaries\Android\UE4Game-Json-armv7-es2.a" Type="StaticLibrary" />
+    <BuildProduct Path="$(EngineDir)\Binaries\Android\UE4Game-Engine-armv7-es2.a" Type="StaticLibrary" />
+    <BuildProduct Path="$(EngineDir)\Binaries\Android\UE4Game-AppFramework-armv7-es2.a" Type="StaticLibrary" />
+    <BuildProduct Path="$(EngineDir)\Binaries\Android\UE4Game-InputCore-armv7-es2.a" Type="StaticLibrary" />
+    <BuildProduct Path="$(EngineDir)\Binaries\Android\UE4Game-SlateReflector-armv7-es2.a" Type="StaticLibrary" />
+    <BuildProduct Path="$(EngineDir)\Binaries\Android\UE4Game-Slate-armv7-es2.a" Type="StaticLibrary" />
+    <BuildProduct Path="$(EngineDir)\Binaries\Android\UE4Game-SlateCore-armv7-es2.a" Type="StaticLibrary" />
+    <BuildProduct Path="$(EngineDir)\Binaries\Android\UE4Game-Networking-armv7-es2.a" Type="StaticLibrary" />
+    <BuildProduct Path="$(EngineDir)\Binaries\Android\UE4Game-Sockets-armv7-es2.a" Type="StaticLibrary" />
+    <BuildProduct Path="$(EngineDir)\Binaries\Android\UE4Game-SandboxFile-armv7-es2.a" Type="StaticLibrary" />
+    <BuildProduct Path="$(EngineDir)\Binaries\Android\UE4Game-VectorVM-armv7-es2.a" Type="StaticLibrary" />
+    <BuildProduct Path="$(EngineDir)\Binaries\Android\UE4Game-Landscape-armv7-es2.a" Type="StaticLibrary" />
+    <BuildProduct Path="$(EngineDir)\Binaries\Android\UE4Game-RenderCore-armv7-es2.a" Type="StaticLibrary" />
+    <BuildProduct Path="$(EngineDir)\Binaries\Android\UE4Game-RHI-armv7-es2.a" Type="StaticLibrary" />
+    <BuildProduct Path="$(EngineDir)\Binaries\Android\UE4Game-NullDrv-armv7-es2.a" Type="StaticLibrary" />
+    <BuildProduct Path="$(EngineDir)\Binaries\Android\UE4Game-ShaderCore-armv7-es2.a" Type="StaticLibrary" />
+    <BuildProduct Path="$(EngineDir)\Binaries\Android\UE4Game-Renderer-armv7-es2.a" Type="StaticLibrary" />
+    <BuildProduct Path="$(EngineDir)\Binaries\Android\UE4Game-UtilityShaders-armv7-es2.a" Type="StaticLibrary" />
+    <BuildProduct Path="$(EngineDir)\Binaries\Android\UE4Game-HeadMountedDisplay-armv7-es2.a" Type="StaticLibrary" />
+    <BuildProduct Path="$(EngineDir)\Binaries\Android\UE4Game-Foliage-armv7-es2.a" Type="StaticLibrary" />
+    <BuildProduct Path="$(EngineDir)\Binaries\Android\UE4Game-UMG-armv7-es2.a" Type="StaticLibrary" />
+    <BuildProduct Path="$(EngineDir)\Binaries\Android\UE4Game-MovieSceneCore-armv7-es2.a" Type="StaticLibrary" />
+    <BuildProduct Path="$(EngineDir)\Binaries\Android\UE4Game-HTTP-armv7-es2.a" Type="StaticLibrary" />
+    <BuildProduct Path="$(EngineDir)\Binaries\Android\UE4Game-ImageWrapper-armv7-es2.a" Type="StaticLibrary" />
+    <BuildProduct Path="$(EngineDir)\Binaries\Android\UE4Game-SlateRHIRenderer-armv7-es2.a" Type="StaticLibrary" />
+    <BuildProduct Path="$(EngineDir)\Binaries\Android\UE4Game-Navmesh-armv7-es2.a" Type="StaticLibrary" />
+    <BuildProduct Path="$(EngineDir)\Binaries\Android\UE4Game-Messaging-armv7-es2.a" Type="StaticLibrary" />
+    <BuildProduct Path="$(EngineDir)\Binaries\Android\UE4Game-AssetRegistry-armv7-es2.a" Type="StaticLibrary" />
+    <BuildProduct Path="$(EngineDir)\Binaries\Android\UE4Game-EngineMessages-armv7-es2.a" Type="StaticLibrary" />
+    <BuildProduct Path="$(EngineDir)\Binaries\Android\UE4Game-EngineSettings-armv7-es2.a" Type="StaticLibrary" />
+    <BuildProduct Path="$(EngineDir)\Binaries\Android\UE4Game-SynthBenchmark-armv7-es2.a" Type="StaticLibrary" />
+    <BuildProduct Path="$(EngineDir)\Binaries\Android\UE4Game-AIModule-armv7-es2.a" Type="StaticLibrary" />
+    <BuildProduct Path="$(EngineDir)\Binaries\Android\UE4Game-GameplayTags-armv7-es2.a" Type="StaticLibrary" />
+    <BuildProduct Path="$(EngineDir)\Binaries\Android\UE4Game-DatabaseSupport-armv7-es2.a" Type="StaticLibrary" />
+    <BuildProduct Path="$(EngineDir)\Binaries\Android\UE4Game-MovieSceneCoreTypes-armv7-es2.a" Type="StaticLibrary" />
+    <BuildProduct Path="$(EngineDir)\Binaries\Android\UE4Game-StreamingPauseRendering-armv7-es2.a" Type="StaticLibrary" />
+    <BuildProduct Path="$(EngineDir)\Binaries\Android\UE4Game-MoviePlayer-armv7-es2.a" Type="StaticLibrary" />
+    <BuildProduct Path="$(EngineDir)\Binaries\Android\UE4Game-Analytics-armv7-es2.a" Type="StaticLibrary" />
+    <BuildProduct Path="$(EngineDir)\Binaries\Android\UE4Game-AnalyticsET-armv7-es2.a" Type="StaticLibrary" />
+    <BuildProduct Path="$(EngineDir)\Binaries\Android\UE4Game-NetworkReplayStreaming-armv7-es2.a" Type="StaticLibrary" />
+    <BuildProduct Path="$(EngineDir)\Binaries\Android\UE4Game-NullNetworkReplayStreaming-armv7-es2.a" Type="StaticLibrary" />
+    <BuildProduct Path="$(EngineDir)\Binaries\Android\UE4Game-HttpNetworkReplayStreaming-armv7-es2.a" Type="StaticLibrary" />
+    <BuildProduct Path="$(EngineDir)\Binaries\Android\UE4Game-OnlineSubsystem-armv7-es2.a" Type="StaticLibrary" />
+    <BuildProduct Path="$(EngineDir)\Binaries\Android\UE4Game-ImageCore-armv7-es2.a" Type="StaticLibrary" />
+    <BuildProduct Path="$(EngineDir)\Binaries\Android\UE4Game-OnlineSubsystemUtils-armv7-es2.a" Type="StaticLibrary" />
+    <BuildProduct Path="$(EngineDir)\Binaries\Android\UE4Game-Voice-armv7-es2.a" Type="StaticLibrary" />
+    <BuildProduct Path="$(EngineDir)\Binaries\Android\UE4Game-Advertising-armv7-es2.a" Type="StaticLibrary" />
+    <BuildProduct Path="$(EngineDir)\Binaries\Android\UE4Game-CrashTracker-armv7-es2.a" Type="StaticLibrary" />
+    <BuildProduct Path="$(EngineDir)\Binaries\Android\UE4Game-GameLiveStreaming-armv7-es2.a" Type="StaticLibrary" />
+    <BuildProduct Path="$(EngineDir)\Binaries\Android\UE4Game-MediaAssets-armv7-es2.a" Type="StaticLibrary" />
+    <BuildProduct Path="$(EngineDir)\Binaries\Android\UE4Game-Media-armv7-es2.a" Type="StaticLibrary" />
+    <BuildProduct Path="$(EngineDir)\Binaries\Android\UE4Game-PakFile-armv7-es2.a" Type="StaticLibrary" />
+    <BuildProduct Path="$(EngineDir)\Binaries\Android\UE4Game-Serialization-armv7-es2.a" Type="StaticLibrary" />
+    <BuildProduct Path="$(EngineDir)\Binaries\Android\UE4Game-NetworkFile-armv7-es2.a" Type="StaticLibrary" />
+    <BuildProduct Path="$(EngineDir)\Binaries\Android\UE4Game-StreamingFile-armv7-es2.a" Type="StaticLibrary" />
+    <BuildProduct Path="$(EngineDir)\Binaries\Android\UE4Game-AutomationWorker-armv7-es2.a" Type="StaticLibrary" />
+    <BuildProduct Path="$(EngineDir)\Binaries\Android\UE4Game-AutomationMessages-armv7-es2.a" Type="StaticLibrary" />
+    <BuildProduct Path="$(EngineDir)\Binaries\Android\UE4Game-OpenGLDrv-armv7-es2.a" Type="StaticLibrary" />
+    <BuildProduct Path="$(EngineDir)\Binaries\Android\UE4Game-AndroidAudio-armv7-es2.a" Type="StaticLibrary" />
+    <BuildProduct Path="$(EngineDir)\Binaries\Android\UE4Game-SessionServices-armv7-es2.a" Type="StaticLibrary" />
+    <BuildProduct Path="$(EngineDir)\Binaries\Android\UE4Game-SessionMessages-armv7-es2.a" Type="StaticLibrary" />
+    <BuildProduct Path="$(EngineDir)\Binaries\Android\UE4Game-TaskGraph-armv7-es2.a" Type="StaticLibrary" />
+    <BuildProduct Path="$(EngineDir)\Binaries\Android\UE4Game-EditorStyle-armv7-es2.a" Type="StaticLibrary" />
+    <BuildProduct Path="$(EngineDir)\Binaries\Android\UE4Game-RealtimeProfiler-armv7-es2.a" Type="StaticLibrary" />
+    <BuildProduct Path="$(EngineDir)\Binaries\Android\UE4Game-ProfilerService-armv7-es2.a" Type="StaticLibrary" />
+    <BuildProduct Path="$(EngineDir)\Binaries\Android\UE4Game-ProfilerMessages-armv7-es2.a" Type="StaticLibrary" />
+    <BuildProduct Path="$(EngineDir)\Binaries\Android\Android\UE4Game-AndroidRuntimeSettings-armv7-es2.a" Type="StaticLibrary" />
+    <BuildProduct Path="$(EngineDir)\Binaries\Android\UE4Game-InputDevice-armv7-es2.a" Type="StaticLibrary" />
+    <BuildProduct Path="$(EngineDir)\Binaries\Android\UE4Game-GameMenuBuilder-armv7-es2.a" Type="StaticLibrary" />
+    <BuildProduct Path="$(EngineDir)\Binaries\Android\UE4Game-XmlParser-armv7-es2.a" Type="StaticLibrary" />
+    <BuildProduct Path="$(EngineDir)\Binaries\Android\UE4Game-AndroidAdvertising-armv7-es2.a" Type="StaticLibrary" />
+  </BuildProducts>
+  <RuntimeDependencies />
+</BuildReceipt>
\ No newline at end of file
diff --git a/Engine/Build/Receipts/UE4Game-Android-Shipping.target.xml b/Engine/Build/Receipts/UE4Game-Android-Shipping.target.xml
new file mode 100644
--- /dev/null
+++ b/Engine/Build/Receipts/UE4Game-Android-Shipping.target.xml
@@ -0,0 +1,92 @@
+<?xml version="1.0" encoding="utf-8"?>
+<BuildReceipt xmlns:xsd="http://www.w3.org/2001/XMLSchema" xmlns:xsi="http://www.w3.org/2001/XMLSchema-instance">
+  <BuildProducts>
+    <BuildProduct Path="$(EngineDir)\Binaries\Android\UE4Game-Android-Shipping-armv7-es2.so" Type="Executable" />
+    <BuildProduct Path="$(EngineDir)\Binaries\Android\UE4Game-Android-Shipping-armv7-es2.apk" Type="Executable" />
+    <BuildProduct Path="$(EngineDir)\Plugins\2D\Paper2D\Binaries\Android\UE4Game-Paper2D-Android-Shipping-Static-armv7-es2.a" Type="StaticLibrary" />
+    <BuildProduct Path="$(EngineDir)\Plugins\Experimental\CharacterAI\Binaries\Android\UE4Game-CharacterAI-Android-Shipping-Static-armv7-es2.a" Type="StaticLibrary" />
+    <BuildProduct Path="$(EngineDir)\Plugins\Experimental\SimpleHMD\Binaries\Android\UE4Game-SimpleHMD-Android-Shipping-Static-armv7-es2.a" Type="StaticLibrary" />
+    <BuildProduct Path="$(EngineDir)\Plugins\Experimental\StructBox\Binaries\Android\UE4Game-StructBox-Android-Shipping-Static-armv7-es2.a" Type="StaticLibrary" />
+    <BuildProduct Path="$(EngineDir)\Plugins\Media\AndroidMedia\Binaries\Android\UE4Game-AndroidMedia-Android-Shipping-Static-armv7-es2.a" Type="StaticLibrary" />
+    <BuildProduct Path="$(EngineDir)\Plugins\Messaging\UdpMessaging\Binaries\Android\UE4Game-UdpMessaging-Android-Shipping-Static-armv7-es2.a" Type="StaticLibrary" />
+    <BuildProduct Path="$(EngineDir)\Plugins\Runtime\Analytics\AnalyticsBlueprintLibrary\Binaries\Android\UE4Game-AnalyticsBlueprintLibrary-Android-Shipping-Static-armv7-es2.a" Type="StaticLibrary" />
+    <BuildProduct Path="$(EngineDir)\Plugins\Runtime\Analytics\FileLogging\Binaries\Android\UE4Game-FileLogging-Android-Shipping-Static-armv7-es2.a" Type="StaticLibrary" />
+    <BuildProduct Path="$(EngineDir)\Plugins\Runtime\AndroidDeviceProfileSelector\Binaries\Android\UE4Game-AndroidDeviceProfileSelector-Android-Shipping-Static-armv7-es2.a" Type="StaticLibrary" />
+    <BuildProduct Path="$(EngineDir)\Plugins\Runtime\AndroidMoviePlayer\Binaries\Android\UE4Game-AndroidMoviePlayer-Android-Shipping-Static-armv7-es2.a" Type="StaticLibrary" />
+    <BuildProduct Path="$(EngineDir)\Plugins\Runtime\CableComponent\Binaries\Android\UE4Game-CableComponent-Android-Shipping-Static-armv7-es2.a" Type="StaticLibrary" />
+    <BuildProduct Path="$(EngineDir)\Plugins\Runtime\CustomMeshComponent\Binaries\Android\UE4Game-CustomMeshComponent-Android-Shipping-Static-armv7-es2.a" Type="StaticLibrary" />
+    <BuildProduct Path="$(EngineDir)\Plugins\Runtime\ExampleDeviceProfileSelector\Binaries\Android\UE4Game-ExampleDeviceProfileSelector-Android-Shipping-Static-armv7-es2.a" Type="StaticLibrary" />
+    <BuildProduct Path="$(EngineDir)\Plugins\Runtime\GearVR\Binaries\Android\UE4Game-GearVR-Android-Shipping-Static-armv7-es2.a" Type="StaticLibrary" />
+    <BuildProduct Path="$(EngineDir)\Plugins\Runtime\SoundMod\Binaries\Android\UE4Game-SoundMod-Android-Shipping-Static-armv7-es2.a" Type="StaticLibrary" />
+    <BuildProduct Path="$(EngineDir)\Plugins\Runtime\SoundVisualizations\Binaries\Android\UE4Game-SoundVisualizations-Android-Shipping-Static-armv7-es2.a" Type="StaticLibrary" />
+    <BuildProduct Path="$(EngineDir)\Plugins\ScriptPlugin\Binaries\Android\UE4Game-ScriptPlugin-Android-Shipping-Static-armv7-es2.a" Type="StaticLibrary" />
+    <BuildProduct Path="$(EngineDir)\Plugins\Slate\SlateRemote\Binaries\Android\UE4Game-SlateRemote-Android-Shipping-Static-armv7-es2.a" Type="StaticLibrary" />
+    <BuildProduct Path="$(EngineDir)\Binaries\Android\UE4Game-Launch-Android-Shipping-armv7-es2.a" Type="StaticLibrary" />
+    <BuildProduct Path="$(EngineDir)\Binaries\Android\UE4Game-Core-Android-Shipping-armv7-es2.a" Type="StaticLibrary" />
+    <BuildProduct Path="$(EngineDir)\Binaries\Android\UE4Game-CoreUObject-Android-Shipping-armv7-es2.a" Type="StaticLibrary" />
+    <BuildProduct Path="$(EngineDir)\Binaries\Android\UE4Game-Projects-Android-Shipping-armv7-es2.a" Type="StaticLibrary" />
+    <BuildProduct Path="$(EngineDir)\Binaries\Android\UE4Game-Json-Android-Shipping-armv7-es2.a" Type="StaticLibrary" />
+    <BuildProduct Path="$(EngineDir)\Binaries\Android\UE4Game-Engine-Android-Shipping-armv7-es2.a" Type="StaticLibrary" />
+    <BuildProduct Path="$(EngineDir)\Binaries\Android\UE4Game-AppFramework-Android-Shipping-armv7-es2.a" Type="StaticLibrary" />
+    <BuildProduct Path="$(EngineDir)\Binaries\Android\UE4Game-InputCore-Android-Shipping-armv7-es2.a" Type="StaticLibrary" />
+    <BuildProduct Path="$(EngineDir)\Binaries\Android\UE4Game-SlateReflector-Android-Shipping-armv7-es2.a" Type="StaticLibrary" />
+    <BuildProduct Path="$(EngineDir)\Binaries\Android\UE4Game-Slate-Android-Shipping-armv7-es2.a" Type="StaticLibrary" />
+    <BuildProduct Path="$(EngineDir)\Binaries\Android\UE4Game-SlateCore-Android-Shipping-armv7-es2.a" Type="StaticLibrary" />
+    <BuildProduct Path="$(EngineDir)\Binaries\Android\UE4Game-Networking-Android-Shipping-armv7-es2.a" Type="StaticLibrary" />
+    <BuildProduct Path="$(EngineDir)\Binaries\Android\UE4Game-Sockets-Android-Shipping-armv7-es2.a" Type="StaticLibrary" />
+    <BuildProduct Path="$(EngineDir)\Binaries\Android\UE4Game-SandboxFile-Android-Shipping-armv7-es2.a" Type="StaticLibrary" />
+    <BuildProduct Path="$(EngineDir)\Binaries\Android\UE4Game-VectorVM-Android-Shipping-armv7-es2.a" Type="StaticLibrary" />
+    <BuildProduct Path="$(EngineDir)\Binaries\Android\UE4Game-Landscape-Android-Shipping-armv7-es2.a" Type="StaticLibrary" />
+    <BuildProduct Path="$(EngineDir)\Binaries\Android\UE4Game-RenderCore-Android-Shipping-armv7-es2.a" Type="StaticLibrary" />
+    <BuildProduct Path="$(EngineDir)\Binaries\Android\UE4Game-RHI-Android-Shipping-armv7-es2.a" Type="StaticLibrary" />
+    <BuildProduct Path="$(EngineDir)\Binaries\Android\UE4Game-NullDrv-Android-Shipping-armv7-es2.a" Type="StaticLibrary" />
+    <BuildProduct Path="$(EngineDir)\Binaries\Android\UE4Game-ShaderCore-Android-Shipping-armv7-es2.a" Type="StaticLibrary" />
+    <BuildProduct Path="$(EngineDir)\Binaries\Android\UE4Game-Renderer-Android-Shipping-armv7-es2.a" Type="StaticLibrary" />
+    <BuildProduct Path="$(EngineDir)\Binaries\Android\UE4Game-UtilityShaders-Android-Shipping-armv7-es2.a" Type="StaticLibrary" />
+    <BuildProduct Path="$(EngineDir)\Binaries\Android\UE4Game-HeadMountedDisplay-Android-Shipping-armv7-es2.a" Type="StaticLibrary" />
+    <BuildProduct Path="$(EngineDir)\Binaries\Android\UE4Game-Foliage-Android-Shipping-armv7-es2.a" Type="StaticLibrary" />
+    <BuildProduct Path="$(EngineDir)\Binaries\Android\UE4Game-UMG-Android-Shipping-armv7-es2.a" Type="StaticLibrary" />
+    <BuildProduct Path="$(EngineDir)\Binaries\Android\UE4Game-MovieSceneCore-Android-Shipping-armv7-es2.a" Type="StaticLibrary" />
+    <BuildProduct Path="$(EngineDir)\Binaries\Android\UE4Game-HTTP-Android-Shipping-armv7-es2.a" Type="StaticLibrary" />
+    <BuildProduct Path="$(EngineDir)\Binaries\Android\UE4Game-ImageWrapper-Android-Shipping-armv7-es2.a" Type="StaticLibrary" />
+    <BuildProduct Path="$(EngineDir)\Binaries\Android\UE4Game-SlateRHIRenderer-Android-Shipping-armv7-es2.a" Type="StaticLibrary" />
+    <BuildProduct Path="$(EngineDir)\Binaries\Android\UE4Game-Navmesh-Android-Shipping-armv7-es2.a" Type="StaticLibrary" />
+    <BuildProduct Path="$(EngineDir)\Binaries\Android\UE4Game-Messaging-Android-Shipping-armv7-es2.a" Type="StaticLibrary" />
+    <BuildProduct Path="$(EngineDir)\Binaries\Android\UE4Game-AssetRegistry-Android-Shipping-armv7-es2.a" Type="StaticLibrary" />
+    <BuildProduct Path="$(EngineDir)\Binaries\Android\UE4Game-EngineMessages-Android-Shipping-armv7-es2.a" Type="StaticLibrary" />
+    <BuildProduct Path="$(EngineDir)\Binaries\Android\UE4Game-EngineSettings-Android-Shipping-armv7-es2.a" Type="StaticLibrary" />
+    <BuildProduct Path="$(EngineDir)\Binaries\Android\UE4Game-SynthBenchmark-Android-Shipping-armv7-es2.a" Type="StaticLibrary" />
+    <BuildProduct Path="$(EngineDir)\Binaries\Android\UE4Game-AIModule-Android-Shipping-armv7-es2.a" Type="StaticLibrary" />
+    <BuildProduct Path="$(EngineDir)\Binaries\Android\UE4Game-GameplayTags-Android-Shipping-armv7-es2.a" Type="StaticLibrary" />
+    <BuildProduct Path="$(EngineDir)\Binaries\Android\UE4Game-DatabaseSupport-Android-Shipping-armv7-es2.a" Type="StaticLibrary" />
+    <BuildProduct Path="$(EngineDir)\Binaries\Android\UE4Game-MovieSceneCoreTypes-Android-Shipping-armv7-es2.a" Type="StaticLibrary" />
+    <BuildProduct Path="$(EngineDir)\Binaries\Android\UE4Game-StreamingPauseRendering-Android-Shipping-armv7-es2.a" Type="StaticLibrary" />
+    <BuildProduct Path="$(EngineDir)\Binaries\Android\UE4Game-MoviePlayer-Android-Shipping-armv7-es2.a" Type="StaticLibrary" />
+    <BuildProduct Path="$(EngineDir)\Binaries\Android\UE4Game-Analytics-Android-Shipping-armv7-es2.a" Type="StaticLibrary" />
+    <BuildProduct Path="$(EngineDir)\Binaries\Android\UE4Game-AnalyticsET-Android-Shipping-armv7-es2.a" Type="StaticLibrary" />
+    <BuildProduct Path="$(EngineDir)\Binaries\Android\UE4Game-NetworkReplayStreaming-Android-Shipping-armv7-es2.a" Type="StaticLibrary" />
+    <BuildProduct Path="$(EngineDir)\Binaries\Android\UE4Game-NullNetworkReplayStreaming-Android-Shipping-armv7-es2.a" Type="StaticLibrary" />
+    <BuildProduct Path="$(EngineDir)\Binaries\Android\UE4Game-HttpNetworkReplayStreaming-Android-Shipping-armv7-es2.a" Type="StaticLibrary" />
+    <BuildProduct Path="$(EngineDir)\Binaries\Android\UE4Game-OnlineSubsystem-Android-Shipping-armv7-es2.a" Type="StaticLibrary" />
+    <BuildProduct Path="$(EngineDir)\Binaries\Android\UE4Game-ImageCore-Android-Shipping-armv7-es2.a" Type="StaticLibrary" />
+    <BuildProduct Path="$(EngineDir)\Binaries\Android\UE4Game-OnlineSubsystemUtils-Android-Shipping-armv7-es2.a" Type="StaticLibrary" />
+    <BuildProduct Path="$(EngineDir)\Binaries\Android\UE4Game-Voice-Android-Shipping-armv7-es2.a" Type="StaticLibrary" />
+    <BuildProduct Path="$(EngineDir)\Binaries\Android\UE4Game-Advertising-Android-Shipping-armv7-es2.a" Type="StaticLibrary" />
+    <BuildProduct Path="$(EngineDir)\Binaries\Android\UE4Game-CrashTracker-Android-Shipping-armv7-es2.a" Type="StaticLibrary" />
+    <BuildProduct Path="$(EngineDir)\Binaries\Android\UE4Game-GameLiveStreaming-Android-Shipping-armv7-es2.a" Type="StaticLibrary" />
+    <BuildProduct Path="$(EngineDir)\Binaries\Android\UE4Game-MediaAssets-Android-Shipping-armv7-es2.a" Type="StaticLibrary" />
+    <BuildProduct Path="$(EngineDir)\Binaries\Android\UE4Game-Media-Android-Shipping-armv7-es2.a" Type="StaticLibrary" />
+    <BuildProduct Path="$(EngineDir)\Binaries\Android\UE4Game-PakFile-Android-Shipping-armv7-es2.a" Type="StaticLibrary" />
+    <BuildProduct Path="$(EngineDir)\Binaries\Android\UE4Game-Serialization-Android-Shipping-armv7-es2.a" Type="StaticLibrary" />
+    <BuildProduct Path="$(EngineDir)\Binaries\Android\UE4Game-OpenGLDrv-Android-Shipping-armv7-es2.a" Type="StaticLibrary" />
+    <BuildProduct Path="$(EngineDir)\Binaries\Android\UE4Game-AndroidAudio-Android-Shipping-armv7-es2.a" Type="StaticLibrary" />
+    <BuildProduct Path="$(EngineDir)\Binaries\Android\UE4Game-SessionServices-Android-Shipping-armv7-es2.a" Type="StaticLibrary" />
+    <BuildProduct Path="$(EngineDir)\Binaries\Android\UE4Game-SessionMessages-Android-Shipping-armv7-es2.a" Type="StaticLibrary" />
+    <BuildProduct Path="$(EngineDir)\Binaries\Android\Android\UE4Game-AndroidRuntimeSettings-Android-Shipping-armv7-es2.a" Type="StaticLibrary" />
+    <BuildProduct Path="$(EngineDir)\Binaries\Android\UE4Game-InputDevice-Android-Shipping-armv7-es2.a" Type="StaticLibrary" />
+    <BuildProduct Path="$(EngineDir)\Binaries\Android\UE4Game-GameMenuBuilder-Android-Shipping-armv7-es2.a" Type="StaticLibrary" />
+    <BuildProduct Path="$(EngineDir)\Binaries\Android\UE4Game-XmlParser-Android-Shipping-armv7-es2.a" Type="StaticLibrary" />
+    <BuildProduct Path="$(EngineDir)\Binaries\Android\UE4Game-AndroidAdvertising-Android-Shipping-armv7-es2.a" Type="StaticLibrary" />
+  </BuildProducts>
+  <RuntimeDependencies />
+</BuildReceipt>
\ No newline at end of file
diff --git a/Engine/Build/Receipts/UE4Game-Android-Test.target.xml b/Engine/Build/Receipts/UE4Game-Android-Test.target.xml
new file mode 100644
--- /dev/null
+++ b/Engine/Build/Receipts/UE4Game-Android-Test.target.xml
@@ -0,0 +1,18 @@
+<?xml version="1.0" encoding="utf-8"?>
+<BuildReceipt xmlns:xsd="http://www.w3.org/2001/XMLSchema" xmlns:xsi="http://www.w3.org/2001/XMLSchema-instance">
+  <BuildProducts>
+    <BuildProduct Path="$(EngineDir)\Binaries\Android\UE4Game-Android-Test-armv7-es2.so" Type="Executable" />
+    <BuildProduct Path="$(EngineDir)\Binaries\Android\UE4Game-Android-Test-armv7-es2.apk" Type="Executable" />
+    <BuildProduct Path="$(EngineDir)\Plugins\2D\Paper2D\Binaries\Android\UE4Game-Paper2D-Android-Test-Static-armv7-es2.a" Type="StaticLibrary" />
+    <BuildProduct Path="$(EngineDir)\Plugins\Experimental\CharacterAI\Binaries\Android\UE4Game-CharacterAI-Android-Test-Static-armv7-es2.a" Type="StaticLibrary" />
+    <BuildProduct Path="$(EngineDir)\Plugins\Media\AndroidMedia\Binaries\Android\UE4Game-AndroidMedia-Android-Test-Static-armv7-es2.a" Type="StaticLibrary" />
+    <BuildProduct Path="$(EngineDir)\Plugins\Messaging\UdpMessaging\Binaries\Android\UE4Game-UdpMessaging-Android-Test-Static-armv7-es2.a" Type="StaticLibrary" />
+    <BuildProduct Path="$(EngineDir)\Plugins\Runtime\AndroidDeviceProfileSelector\Binaries\Android\UE4Game-AndroidDeviceProfileSelector-Android-Test-Static-armv7-es2.a" Type="StaticLibrary" />
+    <BuildProduct Path="$(EngineDir)\Plugins\Runtime\AndroidMoviePlayer\Binaries\Android\UE4Game-AndroidMoviePlayer-Android-Test-Static-armv7-es2.a" Type="StaticLibrary" />
+    <BuildProduct Path="$(EngineDir)\Plugins\Runtime\CableComponent\Binaries\Android\UE4Game-CableComponent-Android-Test-Static-armv7-es2.a" Type="StaticLibrary" />
+    <BuildProduct Path="$(EngineDir)\Plugins\Runtime\CustomMeshComponent\Binaries\Android\UE4Game-CustomMeshComponent-Android-Test-Static-armv7-es2.a" Type="StaticLibrary" />
+    <BuildProduct Path="$(EngineDir)\Plugins\Runtime\ExampleDeviceProfileSelector\Binaries\Android\UE4Game-ExampleDeviceProfileSelector-Android-Test-Static-armv7-es2.a" Type="StaticLibrary" />
+    <BuildProduct Path="$(EngineDir)\Plugins\Slate\SlateRemote\Binaries\Android\UE4Game-SlateRemote-Android-Test-Static-armv7-es2.a" Type="StaticLibrary" />
+  </BuildProducts>
+  <RuntimeDependencies />
+</BuildReceipt>
\ No newline at end of file
diff --git a/Engine/Build/Receipts/UE4Game-HTML5-Development.target.xml b/Engine/Build/Receipts/UE4Game-HTML5-Development.target.xml
new file mode 100644
--- /dev/null
+++ b/Engine/Build/Receipts/UE4Game-HTML5-Development.target.xml
@@ -0,0 +1,93 @@
+<?xml version="1.0" encoding="utf-8"?>
+<BuildReceipt xmlns:xsd="http://www.w3.org/2001/XMLSchema" xmlns:xsi="http://www.w3.org/2001/XMLSchema-instance">
+  <BuildProducts>
+    <BuildProduct Path="$(EngineDir)\Binaries\HTML5\UE4Game.js" Type="Executable" />
+    <BuildProduct Path="$(EngineDir)\Binaries\HTML5\UE4Game.js.mem" Type="RequiredResource" />
+    <BuildProduct Path="$(EngineDir)\Plugins\2D\Paper2D\Binaries\HTML5\UE4Game-Paper2D-Static.bc" Type="StaticLibrary" />
+    <BuildProduct Path="$(EngineDir)\Plugins\Experimental\CharacterAI\Binaries\HTML5\UE4Game-CharacterAI-Static.bc" Type="StaticLibrary" />
+    <BuildProduct Path="$(EngineDir)\Plugins\Experimental\HTML5Networking\Binaries\HTML5\UE4Game-HTML5Networking-Static.bc" Type="StaticLibrary" />
+    <BuildProduct Path="$(EngineDir)\Plugins\Experimental\SimpleHMD\Binaries\HTML5\UE4Game-SimpleHMD-Static.bc" Type="StaticLibrary" />
+    <BuildProduct Path="$(EngineDir)\Plugins\Experimental\StructBox\Binaries\HTML5\UE4Game-StructBox-Static.bc" Type="StaticLibrary" />
+    <BuildProduct Path="$(EngineDir)\Plugins\Runtime\Analytics\AnalyticsBlueprintLibrary\Binaries\HTML5\UE4Game-AnalyticsBlueprintLibrary-Static.bc" Type="StaticLibrary" />
+    <BuildProduct Path="$(EngineDir)\Plugins\Runtime\Analytics\FileLogging\Binaries\HTML5\UE4Game-FileLogging-Static.bc" Type="StaticLibrary" />
+    <BuildProduct Path="$(EngineDir)\Plugins\Runtime\CableComponent\Binaries\HTML5\UE4Game-CableComponent-Static.bc" Type="StaticLibrary" />
+    <BuildProduct Path="$(EngineDir)\Plugins\Runtime\ExampleDeviceProfileSelector\Binaries\HTML5\UE4Game-ExampleDeviceProfileSelector-Static.bc" Type="StaticLibrary" />
+    <BuildProduct Path="$(EngineDir)\Plugins\ScriptPlugin\Binaries\HTML5\UE4Game-ScriptPlugin-Static.bc" Type="StaticLibrary" />
+    <BuildProduct Path="$(EngineDir)\Binaries\HTML5\UE4Game-Launch.bc" Type="StaticLibrary" />
+    <BuildProduct Path="$(EngineDir)\Binaries\HTML5\UE4Game-Core.bc" Type="StaticLibrary" />
+    <BuildProduct Path="$(EngineDir)\Binaries\HTML5\UE4Game-HTML5JS.bc" Type="StaticLibrary" />
+    <BuildProduct Path="$(EngineDir)\Binaries\HTML5\UE4Game-CoreUObject.bc" Type="StaticLibrary" />
+    <BuildProduct Path="$(EngineDir)\Binaries\HTML5\UE4Game-Projects.bc" Type="StaticLibrary" />
+    <BuildProduct Path="$(EngineDir)\Binaries\HTML5\UE4Game-Json.bc" Type="StaticLibrary" />
+    <BuildProduct Path="$(EngineDir)\Binaries\HTML5\UE4Game-Engine.bc" Type="StaticLibrary" />
+    <BuildProduct Path="$(EngineDir)\Binaries\HTML5\UE4Game-AppFramework.bc" Type="StaticLibrary" />
+    <BuildProduct Path="$(EngineDir)\Binaries\HTML5\UE4Game-InputCore.bc" Type="StaticLibrary" />
+    <BuildProduct Path="$(EngineDir)\Binaries\HTML5\UE4Game-SlateReflector.bc" Type="StaticLibrary" />
+    <BuildProduct Path="$(EngineDir)\Binaries\HTML5\UE4Game-Slate.bc" Type="StaticLibrary" />
+    <BuildProduct Path="$(EngineDir)\Binaries\HTML5\UE4Game-SlateCore.bc" Type="StaticLibrary" />
+    <BuildProduct Path="$(EngineDir)\Binaries\HTML5\UE4Game-Networking.bc" Type="StaticLibrary" />
+    <BuildProduct Path="$(EngineDir)\Binaries\HTML5\UE4Game-Sockets.bc" Type="StaticLibrary" />
+    <BuildProduct Path="$(EngineDir)\Binaries\HTML5\UE4Game-SandboxFile.bc" Type="StaticLibrary" />
+    <BuildProduct Path="$(EngineDir)\Binaries\HTML5\UE4Game-VectorVM.bc" Type="StaticLibrary" />
+    <BuildProduct Path="$(EngineDir)\Binaries\HTML5\UE4Game-Landscape.bc" Type="StaticLibrary" />
+    <BuildProduct Path="$(EngineDir)\Binaries\HTML5\UE4Game-RenderCore.bc" Type="StaticLibrary" />
+    <BuildProduct Path="$(EngineDir)\Binaries\HTML5\UE4Game-RHI.bc" Type="StaticLibrary" />
+    <BuildProduct Path="$(EngineDir)\Binaries\HTML5\UE4Game-NullDrv.bc" Type="StaticLibrary" />
+    <BuildProduct Path="$(EngineDir)\Binaries\HTML5\UE4Game-OpenGLDrv.bc" Type="StaticLibrary" />
+    <BuildProduct Path="$(EngineDir)\Binaries\HTML5\UE4Game-ShaderCore.bc" Type="StaticLibrary" />
+    <BuildProduct Path="$(EngineDir)\Binaries\HTML5\UE4Game-UtilityShaders.bc" Type="StaticLibrary" />
+    <BuildProduct Path="$(EngineDir)\Binaries\HTML5\UE4Game-ImageWrapper.bc" Type="StaticLibrary" />
+    <BuildProduct Path="$(EngineDir)\Binaries\HTML5\UE4Game-Renderer.bc" Type="StaticLibrary" />
+    <BuildProduct Path="$(EngineDir)\Binaries\HTML5\UE4Game-HeadMountedDisplay.bc" Type="StaticLibrary" />
+    <BuildProduct Path="$(EngineDir)\Binaries\HTML5\UE4Game-Foliage.bc" Type="StaticLibrary" />
+    <BuildProduct Path="$(EngineDir)\Binaries\HTML5\UE4Game-UMG.bc" Type="StaticLibrary" />
+    <BuildProduct Path="$(EngineDir)\Binaries\HTML5\UE4Game-MovieSceneCore.bc" Type="StaticLibrary" />
+    <BuildProduct Path="$(EngineDir)\Binaries\HTML5\UE4Game-HTTP.bc" Type="StaticLibrary" />
+    <BuildProduct Path="$(EngineDir)\Binaries\HTML5\UE4Game-SlateRHIRenderer.bc" Type="StaticLibrary" />
+    <BuildProduct Path="$(EngineDir)\Binaries\HTML5\UE4Game-Navmesh.bc" Type="StaticLibrary" />
+    <BuildProduct Path="$(EngineDir)\Binaries\HTML5\UE4Game-Messaging.bc" Type="StaticLibrary" />
+    <BuildProduct Path="$(EngineDir)\Binaries\HTML5\UE4Game-AssetRegistry.bc" Type="StaticLibrary" />
+    <BuildProduct Path="$(EngineDir)\Binaries\HTML5\UE4Game-EngineMessages.bc" Type="StaticLibrary" />
+    <BuildProduct Path="$(EngineDir)\Binaries\HTML5\UE4Game-EngineSettings.bc" Type="StaticLibrary" />
+    <BuildProduct Path="$(EngineDir)\Binaries\HTML5\UE4Game-SynthBenchmark.bc" Type="StaticLibrary" />
+    <BuildProduct Path="$(EngineDir)\Binaries\HTML5\UE4Game-AIModule.bc" Type="StaticLibrary" />
+    <BuildProduct Path="$(EngineDir)\Binaries\HTML5\UE4Game-GameplayTags.bc" Type="StaticLibrary" />
+    <BuildProduct Path="$(EngineDir)\Binaries\HTML5\UE4Game-DatabaseSupport.bc" Type="StaticLibrary" />
+    <BuildProduct Path="$(EngineDir)\Binaries\HTML5\UE4Game-MovieSceneCoreTypes.bc" Type="StaticLibrary" />
+    <BuildProduct Path="$(EngineDir)\Binaries\HTML5\UE4Game-StreamingPauseRendering.bc" Type="StaticLibrary" />
+    <BuildProduct Path="$(EngineDir)\Binaries\HTML5\UE4Game-MoviePlayer.bc" Type="StaticLibrary" />
+    <BuildProduct Path="$(EngineDir)\Binaries\HTML5\UE4Game-Analytics.bc" Type="StaticLibrary" />
+    <BuildProduct Path="$(EngineDir)\Binaries\HTML5\UE4Game-AnalyticsET.bc" Type="StaticLibrary" />
+    <BuildProduct Path="$(EngineDir)\Binaries\HTML5\UE4Game-NetworkReplayStreaming.bc" Type="StaticLibrary" />
+    <BuildProduct Path="$(EngineDir)\Binaries\HTML5\UE4Game-NullNetworkReplayStreaming.bc" Type="StaticLibrary" />
+    <BuildProduct Path="$(EngineDir)\Binaries\HTML5\UE4Game-HttpNetworkReplayStreaming.bc" Type="StaticLibrary" />
+    <BuildProduct Path="$(EngineDir)\Binaries\HTML5\UE4Game-OnlineSubsystem.bc" Type="StaticLibrary" />
+    <BuildProduct Path="$(EngineDir)\Binaries\HTML5\UE4Game-ImageCore.bc" Type="StaticLibrary" />
+    <BuildProduct Path="$(EngineDir)\Binaries\HTML5\UE4Game-OnlineSubsystemUtils.bc" Type="StaticLibrary" />
+    <BuildProduct Path="$(EngineDir)\Binaries\HTML5\UE4Game-Voice.bc" Type="StaticLibrary" />
+    <BuildProduct Path="$(EngineDir)\Binaries\HTML5\UE4Game-Advertising.bc" Type="StaticLibrary" />
+    <BuildProduct Path="$(EngineDir)\Binaries\HTML5\UE4Game-CrashTracker.bc" Type="StaticLibrary" />
+    <BuildProduct Path="$(EngineDir)\Binaries\HTML5\UE4Game-GameLiveStreaming.bc" Type="StaticLibrary" />
+    <BuildProduct Path="$(EngineDir)\Binaries\HTML5\UE4Game-MediaAssets.bc" Type="StaticLibrary" />
+    <BuildProduct Path="$(EngineDir)\Binaries\HTML5\UE4Game-Media.bc" Type="StaticLibrary" />
+    <BuildProduct Path="$(EngineDir)\Binaries\HTML5\UE4Game-PakFile.bc" Type="StaticLibrary" />
+    <BuildProduct Path="$(EngineDir)\Binaries\HTML5\UE4Game-Serialization.bc" Type="StaticLibrary" />
+    <BuildProduct Path="$(EngineDir)\Binaries\HTML5\UE4Game-NetworkFile.bc" Type="StaticLibrary" />
+    <BuildProduct Path="$(EngineDir)\Binaries\HTML5\UE4Game-StreamingFile.bc" Type="StaticLibrary" />
+    <BuildProduct Path="$(EngineDir)\Binaries\HTML5\UE4Game-AutomationWorker.bc" Type="StaticLibrary" />
+    <BuildProduct Path="$(EngineDir)\Binaries\HTML5\UE4Game-AutomationMessages.bc" Type="StaticLibrary" />
+    <BuildProduct Path="$(EngineDir)\Binaries\HTML5\UE4Game-HTML5Audio.bc" Type="StaticLibrary" />
+    <BuildProduct Path="$(EngineDir)\Binaries\HTML5\UE4Game-SessionServices.bc" Type="StaticLibrary" />
+    <BuildProduct Path="$(EngineDir)\Binaries\HTML5\UE4Game-SessionMessages.bc" Type="StaticLibrary" />
+    <BuildProduct Path="$(EngineDir)\Binaries\HTML5\UE4Game-TaskGraph.bc" Type="StaticLibrary" />
+    <BuildProduct Path="$(EngineDir)\Binaries\HTML5\UE4Game-EditorStyle.bc" Type="StaticLibrary" />
+    <BuildProduct Path="$(EngineDir)\Binaries\HTML5\UE4Game-RealtimeProfiler.bc" Type="StaticLibrary" />
+    <BuildProduct Path="$(EngineDir)\Binaries\HTML5\UE4Game-ProfilerService.bc" Type="StaticLibrary" />
+    <BuildProduct Path="$(EngineDir)\Binaries\HTML5\UE4Game-ProfilerMessages.bc" Type="StaticLibrary" />
+    <BuildProduct Path="$(EngineDir)\Binaries\HTML5\UE4Game-InputDevice.bc" Type="StaticLibrary" />
+    <BuildProduct Path="$(EngineDir)\Binaries\HTML5\UE4Game-GameMenuBuilder.bc" Type="StaticLibrary" />
+    <BuildProduct Path="$(EngineDir)\Binaries\HTML5\UE4Game-XmlParser.bc" Type="StaticLibrary" />
+    <BuildProduct Path="$(EngineDir)\Binaries\HTML5\UE4Game-OnlineSubsystemNull.bc" Type="StaticLibrary" />
+  </BuildProducts>
+  <RuntimeDependencies />
+</BuildReceipt>
\ No newline at end of file
diff --git a/Engine/Build/Receipts/UE4Game-HTML5-Shipping.target.xml b/Engine/Build/Receipts/UE4Game-HTML5-Shipping.target.xml
new file mode 100644
--- /dev/null
+++ b/Engine/Build/Receipts/UE4Game-HTML5-Shipping.target.xml
@@ -0,0 +1,84 @@
+<?xml version="1.0" encoding="utf-8"?>
+<BuildReceipt xmlns:xsd="http://www.w3.org/2001/XMLSchema" xmlns:xsi="http://www.w3.org/2001/XMLSchema-instance">
+  <BuildProducts>
+    <BuildProduct Path="$(EngineDir)\Binaries\HTML5\UE4Game-HTML5-Shipping.js" Type="Executable" />
+    <BuildProduct Path="$(EngineDir)\Binaries\HTML5\UE4Game-HTML5-Shipping.js.mem" Type="RequiredResource" />
+    <BuildProduct Path="$(EngineDir)\Plugins\2D\Paper2D\Binaries\HTML5\UE4Game-Paper2D-HTML5-Shipping-Static.bc" Type="StaticLibrary" />
+    <BuildProduct Path="$(EngineDir)\Plugins\Experimental\CharacterAI\Binaries\HTML5\UE4Game-CharacterAI-HTML5-Shipping-Static.bc" Type="StaticLibrary" />
+    <BuildProduct Path="$(EngineDir)\Plugins\Experimental\HTML5Networking\Binaries\HTML5\UE4Game-HTML5Networking-HTML5-Shipping-Static.bc" Type="StaticLibrary" />
+    <BuildProduct Path="$(EngineDir)\Plugins\Experimental\SimpleHMD\Binaries\HTML5\UE4Game-SimpleHMD-HTML5-Shipping-Static.bc" Type="StaticLibrary" />
+    <BuildProduct Path="$(EngineDir)\Plugins\Experimental\StructBox\Binaries\HTML5\UE4Game-StructBox-HTML5-Shipping-Static.bc" Type="StaticLibrary" />
+    <BuildProduct Path="$(EngineDir)\Plugins\Runtime\Analytics\AnalyticsBlueprintLibrary\Binaries\HTML5\UE4Game-AnalyticsBlueprintLibrary-HTML5-Shipping-Static.bc" Type="StaticLibrary" />
+    <BuildProduct Path="$(EngineDir)\Plugins\Runtime\Analytics\FileLogging\Binaries\HTML5\UE4Game-FileLogging-HTML5-Shipping-Static.bc" Type="StaticLibrary" />
+    <BuildProduct Path="$(EngineDir)\Plugins\Runtime\CableComponent\Binaries\HTML5\UE4Game-CableComponent-HTML5-Shipping-Static.bc" Type="StaticLibrary" />
+    <BuildProduct Path="$(EngineDir)\Plugins\Runtime\ExampleDeviceProfileSelector\Binaries\HTML5\UE4Game-ExampleDeviceProfileSelector-HTML5-Shipping-Static.bc" Type="StaticLibrary" />
+    <BuildProduct Path="$(EngineDir)\Plugins\ScriptPlugin\Binaries\HTML5\UE4Game-ScriptPlugin-HTML5-Shipping-Static.bc" Type="StaticLibrary" />
+    <BuildProduct Path="$(EngineDir)\Binaries\HTML5\UE4Game-Launch-HTML5-Shipping.bc" Type="StaticLibrary" />
+    <BuildProduct Path="$(EngineDir)\Binaries\HTML5\UE4Game-Core-HTML5-Shipping.bc" Type="StaticLibrary" />
+    <BuildProduct Path="$(EngineDir)\Binaries\HTML5\UE4Game-HTML5JS-HTML5-Shipping.bc" Type="StaticLibrary" />
+    <BuildProduct Path="$(EngineDir)\Binaries\HTML5\UE4Game-CoreUObject-HTML5-Shipping.bc" Type="StaticLibrary" />
+    <BuildProduct Path="$(EngineDir)\Binaries\HTML5\UE4Game-Projects-HTML5-Shipping.bc" Type="StaticLibrary" />
+    <BuildProduct Path="$(EngineDir)\Binaries\HTML5\UE4Game-Json-HTML5-Shipping.bc" Type="StaticLibrary" />
+    <BuildProduct Path="$(EngineDir)\Binaries\HTML5\UE4Game-Engine-HTML5-Shipping.bc" Type="StaticLibrary" />
+    <BuildProduct Path="$(EngineDir)\Binaries\HTML5\UE4Game-AppFramework-HTML5-Shipping.bc" Type="StaticLibrary" />
+    <BuildProduct Path="$(EngineDir)\Binaries\HTML5\UE4Game-InputCore-HTML5-Shipping.bc" Type="StaticLibrary" />
+    <BuildProduct Path="$(EngineDir)\Binaries\HTML5\UE4Game-SlateReflector-HTML5-Shipping.bc" Type="StaticLibrary" />
+    <BuildProduct Path="$(EngineDir)\Binaries\HTML5\UE4Game-Slate-HTML5-Shipping.bc" Type="StaticLibrary" />
+    <BuildProduct Path="$(EngineDir)\Binaries\HTML5\UE4Game-SlateCore-HTML5-Shipping.bc" Type="StaticLibrary" />
+    <BuildProduct Path="$(EngineDir)\Binaries\HTML5\UE4Game-Networking-HTML5-Shipping.bc" Type="StaticLibrary" />
+    <BuildProduct Path="$(EngineDir)\Binaries\HTML5\UE4Game-Sockets-HTML5-Shipping.bc" Type="StaticLibrary" />
+    <BuildProduct Path="$(EngineDir)\Binaries\HTML5\UE4Game-SandboxFile-HTML5-Shipping.bc" Type="StaticLibrary" />
+    <BuildProduct Path="$(EngineDir)\Binaries\HTML5\UE4Game-VectorVM-HTML5-Shipping.bc" Type="StaticLibrary" />
+    <BuildProduct Path="$(EngineDir)\Binaries\HTML5\UE4Game-Landscape-HTML5-Shipping.bc" Type="StaticLibrary" />
+    <BuildProduct Path="$(EngineDir)\Binaries\HTML5\UE4Game-RenderCore-HTML5-Shipping.bc" Type="StaticLibrary" />
+    <BuildProduct Path="$(EngineDir)\Binaries\HTML5\UE4Game-RHI-HTML5-Shipping.bc" Type="StaticLibrary" />
+    <BuildProduct Path="$(EngineDir)\Binaries\HTML5\UE4Game-NullDrv-HTML5-Shipping.bc" Type="StaticLibrary" />
+    <BuildProduct Path="$(EngineDir)\Binaries\HTML5\UE4Game-OpenGLDrv-HTML5-Shipping.bc" Type="StaticLibrary" />
+    <BuildProduct Path="$(EngineDir)\Binaries\HTML5\UE4Game-ShaderCore-HTML5-Shipping.bc" Type="StaticLibrary" />
+    <BuildProduct Path="$(EngineDir)\Binaries\HTML5\UE4Game-UtilityShaders-HTML5-Shipping.bc" Type="StaticLibrary" />
+    <BuildProduct Path="$(EngineDir)\Binaries\HTML5\UE4Game-ImageWrapper-HTML5-Shipping.bc" Type="StaticLibrary" />
+    <BuildProduct Path="$(EngineDir)\Binaries\HTML5\UE4Game-Renderer-HTML5-Shipping.bc" Type="StaticLibrary" />
+    <BuildProduct Path="$(EngineDir)\Binaries\HTML5\UE4Game-HeadMountedDisplay-HTML5-Shipping.bc" Type="StaticLibrary" />
+    <BuildProduct Path="$(EngineDir)\Binaries\HTML5\UE4Game-Foliage-HTML5-Shipping.bc" Type="StaticLibrary" />
+    <BuildProduct Path="$(EngineDir)\Binaries\HTML5\UE4Game-UMG-HTML5-Shipping.bc" Type="StaticLibrary" />
+    <BuildProduct Path="$(EngineDir)\Binaries\HTML5\UE4Game-MovieSceneCore-HTML5-Shipping.bc" Type="StaticLibrary" />
+    <BuildProduct Path="$(EngineDir)\Binaries\HTML5\UE4Game-HTTP-HTML5-Shipping.bc" Type="StaticLibrary" />
+    <BuildProduct Path="$(EngineDir)\Binaries\HTML5\UE4Game-SlateRHIRenderer-HTML5-Shipping.bc" Type="StaticLibrary" />
+    <BuildProduct Path="$(EngineDir)\Binaries\HTML5\UE4Game-Navmesh-HTML5-Shipping.bc" Type="StaticLibrary" />
+    <BuildProduct Path="$(EngineDir)\Binaries\HTML5\UE4Game-Messaging-HTML5-Shipping.bc" Type="StaticLibrary" />
+    <BuildProduct Path="$(EngineDir)\Binaries\HTML5\UE4Game-AssetRegistry-HTML5-Shipping.bc" Type="StaticLibrary" />
+    <BuildProduct Path="$(EngineDir)\Binaries\HTML5\UE4Game-EngineMessages-HTML5-Shipping.bc" Type="StaticLibrary" />
+    <BuildProduct Path="$(EngineDir)\Binaries\HTML5\UE4Game-EngineSettings-HTML5-Shipping.bc" Type="StaticLibrary" />
+    <BuildProduct Path="$(EngineDir)\Binaries\HTML5\UE4Game-SynthBenchmark-HTML5-Shipping.bc" Type="StaticLibrary" />
+    <BuildProduct Path="$(EngineDir)\Binaries\HTML5\UE4Game-AIModule-HTML5-Shipping.bc" Type="StaticLibrary" />
+    <BuildProduct Path="$(EngineDir)\Binaries\HTML5\UE4Game-GameplayTags-HTML5-Shipping.bc" Type="StaticLibrary" />
+    <BuildProduct Path="$(EngineDir)\Binaries\HTML5\UE4Game-DatabaseSupport-HTML5-Shipping.bc" Type="StaticLibrary" />
+    <BuildProduct Path="$(EngineDir)\Binaries\HTML5\UE4Game-MovieSceneCoreTypes-HTML5-Shipping.bc" Type="StaticLibrary" />
+    <BuildProduct Path="$(EngineDir)\Binaries\HTML5\UE4Game-StreamingPauseRendering-HTML5-Shipping.bc" Type="StaticLibrary" />
+    <BuildProduct Path="$(EngineDir)\Binaries\HTML5\UE4Game-MoviePlayer-HTML5-Shipping.bc" Type="StaticLibrary" />
+    <BuildProduct Path="$(EngineDir)\Binaries\HTML5\UE4Game-Analytics-HTML5-Shipping.bc" Type="StaticLibrary" />
+    <BuildProduct Path="$(EngineDir)\Binaries\HTML5\UE4Game-AnalyticsET-HTML5-Shipping.bc" Type="StaticLibrary" />
+    <BuildProduct Path="$(EngineDir)\Binaries\HTML5\UE4Game-NetworkReplayStreaming-HTML5-Shipping.bc" Type="StaticLibrary" />
+    <BuildProduct Path="$(EngineDir)\Binaries\HTML5\UE4Game-NullNetworkReplayStreaming-HTML5-Shipping.bc" Type="StaticLibrary" />
+    <BuildProduct Path="$(EngineDir)\Binaries\HTML5\UE4Game-HttpNetworkReplayStreaming-HTML5-Shipping.bc" Type="StaticLibrary" />
+    <BuildProduct Path="$(EngineDir)\Binaries\HTML5\UE4Game-OnlineSubsystem-HTML5-Shipping.bc" Type="StaticLibrary" />
+    <BuildProduct Path="$(EngineDir)\Binaries\HTML5\UE4Game-ImageCore-HTML5-Shipping.bc" Type="StaticLibrary" />
+    <BuildProduct Path="$(EngineDir)\Binaries\HTML5\UE4Game-OnlineSubsystemUtils-HTML5-Shipping.bc" Type="StaticLibrary" />
+    <BuildProduct Path="$(EngineDir)\Binaries\HTML5\UE4Game-Voice-HTML5-Shipping.bc" Type="StaticLibrary" />
+    <BuildProduct Path="$(EngineDir)\Binaries\HTML5\UE4Game-Advertising-HTML5-Shipping.bc" Type="StaticLibrary" />
+    <BuildProduct Path="$(EngineDir)\Binaries\HTML5\UE4Game-CrashTracker-HTML5-Shipping.bc" Type="StaticLibrary" />
+    <BuildProduct Path="$(EngineDir)\Binaries\HTML5\UE4Game-GameLiveStreaming-HTML5-Shipping.bc" Type="StaticLibrary" />
+    <BuildProduct Path="$(EngineDir)\Binaries\HTML5\UE4Game-MediaAssets-HTML5-Shipping.bc" Type="StaticLibrary" />
+    <BuildProduct Path="$(EngineDir)\Binaries\HTML5\UE4Game-Media-HTML5-Shipping.bc" Type="StaticLibrary" />
+    <BuildProduct Path="$(EngineDir)\Binaries\HTML5\UE4Game-PakFile-HTML5-Shipping.bc" Type="StaticLibrary" />
+    <BuildProduct Path="$(EngineDir)\Binaries\HTML5\UE4Game-Serialization-HTML5-Shipping.bc" Type="StaticLibrary" />
+    <BuildProduct Path="$(EngineDir)\Binaries\HTML5\UE4Game-HTML5Audio-HTML5-Shipping.bc" Type="StaticLibrary" />
+    <BuildProduct Path="$(EngineDir)\Binaries\HTML5\UE4Game-SessionServices-HTML5-Shipping.bc" Type="StaticLibrary" />
+    <BuildProduct Path="$(EngineDir)\Binaries\HTML5\UE4Game-SessionMessages-HTML5-Shipping.bc" Type="StaticLibrary" />
+    <BuildProduct Path="$(EngineDir)\Binaries\HTML5\UE4Game-InputDevice-HTML5-Shipping.bc" Type="StaticLibrary" />
+    <BuildProduct Path="$(EngineDir)\Binaries\HTML5\UE4Game-GameMenuBuilder-HTML5-Shipping.bc" Type="StaticLibrary" />
+    <BuildProduct Path="$(EngineDir)\Binaries\HTML5\UE4Game-XmlParser-HTML5-Shipping.bc" Type="StaticLibrary" />
+    <BuildProduct Path="$(EngineDir)\Binaries\HTML5\UE4Game-OnlineSubsystemNull-HTML5-Shipping.bc" Type="StaticLibrary" />
+  </BuildProducts>
+  <RuntimeDependencies />
+</BuildReceipt>
\ No newline at end of file
diff --git a/Engine/Build/Receipts/UE4Game-HTML5-Test.target.xml b/Engine/Build/Receipts/UE4Game-HTML5-Test.target.xml
new file mode 100644
--- /dev/null
+++ b/Engine/Build/Receipts/UE4Game-HTML5-Test.target.xml
@@ -0,0 +1,12 @@
+<?xml version="1.0" encoding="utf-8"?>
+<BuildReceipt xmlns:xsd="http://www.w3.org/2001/XMLSchema" xmlns:xsi="http://www.w3.org/2001/XMLSchema-instance">
+  <BuildProducts>
+    <BuildProduct Path="$(EngineDir)\Binaries\HTML5\UE4Game-HTML5-Test.js" Type="Executable" />
+    <BuildProduct Path="$(EngineDir)\Binaries\HTML5\UE4Game-HTML5-Test.js.mem" Type="RequiredResource" />
+    <BuildProduct Path="$(EngineDir)\Plugins\2D\Paper2D\Binaries\HTML5\UE4Game-Paper2D-HTML5-Test-Static.bc" Type="StaticLibrary" />
+    <BuildProduct Path="$(EngineDir)\Plugins\Experimental\CharacterAI\Binaries\HTML5\UE4Game-CharacterAI-HTML5-Test-Static.bc" Type="StaticLibrary" />
+    <BuildProduct Path="$(EngineDir)\Plugins\Runtime\CableComponent\Binaries\HTML5\UE4Game-CableComponent-HTML5-Test-Static.bc" Type="StaticLibrary" />
+    <BuildProduct Path="$(EngineDir)\Plugins\Runtime\ExampleDeviceProfileSelector\Binaries\HTML5\UE4Game-ExampleDeviceProfileSelector-HTML5-Test-Static.bc" Type="StaticLibrary" />
+  </BuildProducts>
+  <RuntimeDependencies />
+</BuildReceipt>
\ No newline at end of file
