GUBP Node Shared_LabelPromotable built from changelist 2505032

diff --git a/Engine/Build/Receipts/UE4Game-IOS-Development.target.xml b/Engine/Build/Receipts/UE4Game-IOS-Development.target.xml
new file mode 100644
--- /dev/null
+++ b/Engine/Build/Receipts/UE4Game-IOS-Development.target.xml
@@ -0,0 +1,104 @@
+<?xml version="1.0" encoding="utf-8"?>
+<BuildReceipt xmlns:xsi="http://www.w3.org/2001/XMLSchema-instance" xmlns:xsd="http://www.w3.org/2001/XMLSchema">
+  <BuildProducts>
+    <BuildProduct Path="$(EngineDir)/Binaries/IOS/UE4Game" Type="Executable" />
+    <BuildProduct Path="$(EngineDir)/Binaries/IOS/UE4Game.stub" Type="Executable" />
+    <BuildProduct Path="$(EngineDir)/Plugins/2D/Paper2D/Binaries/IOS/UE4Game-Paper2D-Static.a" Type="StaticLibrary" />
+    <BuildProduct Path="$(EngineDir)/Plugins/Experimental/CharacterAI/Binaries/IOS/UE4Game-CharacterAI-Static.a" Type="StaticLibrary" />
+    <BuildProduct Path="$(EngineDir)/Plugins/Experimental/SimpleHMD/Binaries/IOS/UE4Game-SimpleHMD-Static.a" Type="StaticLibrary" />
+    <BuildProduct Path="$(EngineDir)/Plugins/Experimental/StructBox/Binaries/IOS/UE4Game-StructBox-Static.a" Type="StaticLibrary" />
+    <BuildProduct Path="$(EngineDir)/Plugins/Media/AvfMedia/Binaries/IOS/UE4Game-AvfMedia-Static.a" Type="StaticLibrary" />
+    <BuildProduct Path="$(EngineDir)/Plugins/Messaging/UdpMessaging/Binaries/IOS/UE4Game-UdpMessaging-Static.a" Type="StaticLibrary" />
+    <BuildProduct Path="$(EngineDir)/Plugins/Runtime/Advertising/IOSTapJoy/Binaries/IOS/UE4Game-IOSTapJoy-Static.a" Type="StaticLibrary" />
+    <BuildProduct Path="$(EngineDir)/Plugins/Runtime/Analytics/AnalyticsBlueprintLibrary/Binaries/IOS/UE4Game-AnalyticsBlueprintLibrary-Static.a" Type="StaticLibrary" />
+    <BuildProduct Path="$(EngineDir)/Plugins/Runtime/Analytics/Apsalar/IOSApsalar/Binaries/IOS/NotForLicensees/UE4Game-IOSApsalar-Static.a" Type="StaticLibrary" />
+    <BuildProduct Path="$(EngineDir)/Plugins/Runtime/Analytics/FileLogging/Binaries/IOS/UE4Game-FileLogging-Static.a" Type="StaticLibrary" />
+    <BuildProduct Path="$(EngineDir)/Plugins/Runtime/AppleMoviePlayer/Binaries/IOS/UE4Game-AppleMoviePlayer-Static.a" Type="StaticLibrary" />
+    <BuildProduct Path="$(EngineDir)/Plugins/Runtime/CableComponent/Binaries/IOS/UE4Game-CableComponent-Static.a" Type="StaticLibrary" />
+    <BuildProduct Path="$(EngineDir)/Plugins/Runtime/CustomMeshComponent/Binaries/IOS/UE4Game-CustomMeshComponent-Static.a" Type="StaticLibrary" />
+    <BuildProduct Path="$(EngineDir)/Plugins/Runtime/ExampleDeviceProfileSelector/Binaries/IOS/UE4Game-ExampleDeviceProfileSelector-Static.a" Type="StaticLibrary" />
+    <BuildProduct Path="$(EngineDir)/Plugins/Runtime/IOSDeviceProfileSelector/Binaries/IOS/UE4Game-IOSDeviceProfileSelector-Static.a" Type="StaticLibrary" />
+    <BuildProduct Path="$(EngineDir)/Plugins/Runtime/SoundVisualizations/Binaries/IOS/UE4Game-SoundVisualizations-Static.a" Type="StaticLibrary" />
+    <BuildProduct Path="$(EngineDir)/Plugins/ScriptPlugin/Binaries/IOS/UE4Game-ScriptPlugin-Static.a" Type="StaticLibrary" />
+    <BuildProduct Path="$(EngineDir)/Plugins/Slate/SlateRemote/Binaries/IOS/UE4Game-SlateRemote-Static.a" Type="StaticLibrary" />
+    <BuildProduct Path="$(EngineDir)/Binaries/IOS/UE4Game-Launch.a" Type="StaticLibrary" />
+    <BuildProduct Path="$(EngineDir)/Binaries/IOS/UE4Game-Core.a" Type="StaticLibrary" />
+    <BuildProduct Path="$(EngineDir)/Binaries/IOS/UE4Game-CoreUObject.a" Type="StaticLibrary" />
+    <BuildProduct Path="$(EngineDir)/Binaries/IOS/UE4Game-Projects.a" Type="StaticLibrary" />
+    <BuildProduct Path="$(EngineDir)/Binaries/IOS/UE4Game-Json.a" Type="StaticLibrary" />
+    <BuildProduct Path="$(EngineDir)/Binaries/IOS/UE4Game-Engine.a" Type="StaticLibrary" />
+    <BuildProduct Path="$(EngineDir)/Binaries/IOS/UE4Game-AppFramework.a" Type="StaticLibrary" />
+    <BuildProduct Path="$(EngineDir)/Binaries/IOS/UE4Game-InputCore.a" Type="StaticLibrary" />
+    <BuildProduct Path="$(EngineDir)/Binaries/IOS/UE4Game-SlateReflector.a" Type="StaticLibrary" />
+    <BuildProduct Path="$(EngineDir)/Binaries/IOS/UE4Game-Slate.a" Type="StaticLibrary" />
+    <BuildProduct Path="$(EngineDir)/Binaries/IOS/UE4Game-SlateCore.a" Type="StaticLibrary" />
+    <BuildProduct Path="$(EngineDir)/Binaries/IOS/UE4Game-Networking.a" Type="StaticLibrary" />
+    <BuildProduct Path="$(EngineDir)/Binaries/IOS/UE4Game-Sockets.a" Type="StaticLibrary" />
+    <BuildProduct Path="$(EngineDir)/Binaries/IOS/UE4Game-SandboxFile.a" Type="StaticLibrary" />
+    <BuildProduct Path="$(EngineDir)/Binaries/IOS/UE4Game-VectorVM.a" Type="StaticLibrary" />
+    <BuildProduct Path="$(EngineDir)/Binaries/IOS/UE4Game-Landscape.a" Type="StaticLibrary" />
+    <BuildProduct Path="$(EngineDir)/Binaries/IOS/UE4Game-RenderCore.a" Type="StaticLibrary" />
+    <BuildProduct Path="$(EngineDir)/Binaries/IOS/UE4Game-RHI.a" Type="StaticLibrary" />
+    <BuildProduct Path="$(EngineDir)/Binaries/IOS/UE4Game-NullDrv.a" Type="StaticLibrary" />
+    <BuildProduct Path="$(EngineDir)/Binaries/IOS/UE4Game-ShaderCore.a" Type="StaticLibrary" />
+    <BuildProduct Path="$(EngineDir)/Binaries/IOS/UE4Game-Renderer.a" Type="StaticLibrary" />
+    <BuildProduct Path="$(EngineDir)/Binaries/IOS/UE4Game-UtilityShaders.a" Type="StaticLibrary" />
+    <BuildProduct Path="$(EngineDir)/Binaries/IOS/UE4Game-HeadMountedDisplay.a" Type="StaticLibrary" />
+    <BuildProduct Path="$(EngineDir)/Binaries/IOS/UE4Game-Foliage.a" Type="StaticLibrary" />
+    <BuildProduct Path="$(EngineDir)/Binaries/IOS/UE4Game-UMG.a" Type="StaticLibrary" />
+    <BuildProduct Path="$(EngineDir)/Binaries/IOS/UE4Game-MovieSceneCore.a" Type="StaticLibrary" />
+    <BuildProduct Path="$(EngineDir)/Binaries/IOS/UE4Game-HTTP.a" Type="StaticLibrary" />
+    <BuildProduct Path="$(EngineDir)/Binaries/IOS/UE4Game-ImageWrapper.a" Type="StaticLibrary" />
+    <BuildProduct Path="$(EngineDir)/Binaries/IOS/UE4Game-SlateRHIRenderer.a" Type="StaticLibrary" />
+    <BuildProduct Path="$(EngineDir)/Binaries/IOS/UE4Game-Navmesh.a" Type="StaticLibrary" />
+    <BuildProduct Path="$(EngineDir)/Binaries/IOS/UE4Game-Messaging.a" Type="StaticLibrary" />
+    <BuildProduct Path="$(EngineDir)/Binaries/IOS/UE4Game-AssetRegistry.a" Type="StaticLibrary" />
+    <BuildProduct Path="$(EngineDir)/Binaries/IOS/UE4Game-EngineMessages.a" Type="StaticLibrary" />
+    <BuildProduct Path="$(EngineDir)/Binaries/IOS/UE4Game-EngineSettings.a" Type="StaticLibrary" />
+    <BuildProduct Path="$(EngineDir)/Binaries/IOS/UE4Game-SynthBenchmark.a" Type="StaticLibrary" />
+    <BuildProduct Path="$(EngineDir)/Binaries/IOS/UE4Game-AIModule.a" Type="StaticLibrary" />
+    <BuildProduct Path="$(EngineDir)/Binaries/IOS/UE4Game-GameplayTags.a" Type="StaticLibrary" />
+    <BuildProduct Path="$(EngineDir)/Binaries/IOS/UE4Game-DatabaseSupport.a" Type="StaticLibrary" />
+    <BuildProduct Path="$(EngineDir)/Binaries/IOS/UE4Game-MovieSceneCoreTypes.a" Type="StaticLibrary" />
+    <BuildProduct Path="$(EngineDir)/Binaries/IOS/UE4Game-StreamingPauseRendering.a" Type="StaticLibrary" />
+    <BuildProduct Path="$(EngineDir)/Binaries/IOS/UE4Game-MoviePlayer.a" Type="StaticLibrary" />
+    <BuildProduct Path="$(EngineDir)/Binaries/IOS/UE4Game-Analytics.a" Type="StaticLibrary" />
+    <BuildProduct Path="$(EngineDir)/Binaries/IOS/UE4Game-AnalyticsET.a" Type="StaticLibrary" />
+    <BuildProduct Path="$(EngineDir)/Binaries/IOS/UE4Game-NetworkReplayStreaming.a" Type="StaticLibrary" />
+    <BuildProduct Path="$(EngineDir)/Binaries/IOS/UE4Game-NullNetworkReplayStreaming.a" Type="StaticLibrary" />
+    <BuildProduct Path="$(EngineDir)/Binaries/IOS/UE4Game-HttpNetworkReplayStreaming.a" Type="StaticLibrary" />
+    <BuildProduct Path="$(EngineDir)/Binaries/IOS/UE4Game-OnlineSubsystem.a" Type="StaticLibrary" />
+    <BuildProduct Path="$(EngineDir)/Binaries/IOS/UE4Game-ImageCore.a" Type="StaticLibrary" />
+    <BuildProduct Path="$(EngineDir)/Binaries/IOS/UE4Game-OnlineSubsystemUtils.a" Type="StaticLibrary" />
+    <BuildProduct Path="$(EngineDir)/Binaries/IOS/UE4Game-Voice.a" Type="StaticLibrary" />
+    <BuildProduct Path="$(EngineDir)/Binaries/IOS/UE4Game-Advertising.a" Type="StaticLibrary" />
+    <BuildProduct Path="$(EngineDir)/Binaries/IOS/UE4Game-CrashTracker.a" Type="StaticLibrary" />
+    <BuildProduct Path="$(EngineDir)/Binaries/IOS/UE4Game-GameLiveStreaming.a" Type="StaticLibrary" />
+    <BuildProduct Path="$(EngineDir)/Binaries/IOS/UE4Game-MediaAssets.a" Type="StaticLibrary" />
+    <BuildProduct Path="$(EngineDir)/Binaries/IOS/UE4Game-Media.a" Type="StaticLibrary" />
+    <BuildProduct Path="$(EngineDir)/Binaries/IOS/UE4Game-PakFile.a" Type="StaticLibrary" />
+    <BuildProduct Path="$(EngineDir)/Binaries/IOS/UE4Game-Serialization.a" Type="StaticLibrary" />
+    <BuildProduct Path="$(EngineDir)/Binaries/IOS/UE4Game-NetworkFile.a" Type="StaticLibrary" />
+    <BuildProduct Path="$(EngineDir)/Binaries/IOS/UE4Game-StreamingFile.a" Type="StaticLibrary" />
+    <BuildProduct Path="$(EngineDir)/Binaries/IOS/UE4Game-AutomationWorker.a" Type="StaticLibrary" />
+    <BuildProduct Path="$(EngineDir)/Binaries/IOS/UE4Game-AutomationMessages.a" Type="StaticLibrary" />
+    <BuildProduct Path="$(EngineDir)/Binaries/IOS/UE4Game-OpenGLDrv.a" Type="StaticLibrary" />
+    <BuildProduct Path="$(EngineDir)/Binaries/IOS/UE4Game-LaunchDaemonMessages.a" Type="StaticLibrary" />
+    <BuildProduct Path="$(EngineDir)/Binaries/IOS/UE4Game-SessionServices.a" Type="StaticLibrary" />
+    <BuildProduct Path="$(EngineDir)/Binaries/IOS/UE4Game-SessionMessages.a" Type="StaticLibrary" />
+    <BuildProduct Path="$(EngineDir)/Binaries/IOS/UE4Game-TaskGraph.a" Type="StaticLibrary" />
+    <BuildProduct Path="$(EngineDir)/Binaries/IOS/UE4Game-EditorStyle.a" Type="StaticLibrary" />
+    <BuildProduct Path="$(EngineDir)/Binaries/IOS/UE4Game-RealtimeProfiler.a" Type="StaticLibrary" />
+    <BuildProduct Path="$(EngineDir)/Binaries/IOS/UE4Game-ProfilerService.a" Type="StaticLibrary" />
+    <BuildProduct Path="$(EngineDir)/Binaries/IOS/UE4Game-ProfilerMessages.a" Type="StaticLibrary" />
+    <BuildProduct Path="$(EngineDir)/Binaries/IOS/IOS/UE4Game-IOSAudio.a" Type="StaticLibrary" />
+    <BuildProduct Path="$(EngineDir)/Binaries/IOS/IOS/UE4Game-IOSRuntimeSettings.a" Type="StaticLibrary" />
+    <BuildProduct Path="$(EngineDir)/Binaries/IOS/UE4Game-InputDevice.a" Type="StaticLibrary" />
+    <BuildProduct Path="$(EngineDir)/Binaries/IOS/UE4Game-GameMenuBuilder.a" Type="StaticLibrary" />
+    <BuildProduct Path="$(EngineDir)/Binaries/IOS/UE4Game-XmlParser.a" Type="StaticLibrary" />
+    <BuildProduct Path="$(EngineDir)/Binaries/IOS/UE4Game-OnlineSubsystemFacebook.a" Type="StaticLibrary" />
+    <BuildProduct Path="$(EngineDir)/Binaries/IOS/UE4Game-OnlineSubsystemIOS.a" Type="StaticLibrary" />
+    <BuildProduct Path="$(EngineDir)/Binaries/IOS/UE4Game-IOSAdvertising.a" Type="StaticLibrary" />
+  </BuildProducts>
+  <RuntimeDependencies />
+</BuildReceipt>
\ No newline at end of file
diff --git a/Engine/Build/Receipts/UE4Game-IOS-Shipping.target.xml b/Engine/Build/Receipts/UE4Game-IOS-Shipping.target.xml
new file mode 100644
--- /dev/null
+++ b/Engine/Build/Receipts/UE4Game-IOS-Shipping.target.xml
@@ -0,0 +1,95 @@
+<?xml version="1.0" encoding="utf-8"?>
+<BuildReceipt xmlns:xsi="http://www.w3.org/2001/XMLSchema-instance" xmlns:xsd="http://www.w3.org/2001/XMLSchema">
+  <BuildProducts>
+    <BuildProduct Path="$(EngineDir)/Binaries/IOS/UE4Game-IOS-Shipping" Type="Executable" />
+    <BuildProduct Path="$(EngineDir)/Binaries/IOS/UE4Game-IOS-Shipping.stub" Type="Executable" />
+    <BuildProduct Path="$(EngineDir)/Plugins/2D/Paper2D/Binaries/IOS/UE4Game-Paper2D-IOS-Shipping-Static.a" Type="StaticLibrary" />
+    <BuildProduct Path="$(EngineDir)/Plugins/Experimental/CharacterAI/Binaries/IOS/UE4Game-CharacterAI-IOS-Shipping-Static.a" Type="StaticLibrary" />
+    <BuildProduct Path="$(EngineDir)/Plugins/Experimental/SimpleHMD/Binaries/IOS/UE4Game-SimpleHMD-IOS-Shipping-Static.a" Type="StaticLibrary" />
+    <BuildProduct Path="$(EngineDir)/Plugins/Experimental/StructBox/Binaries/IOS/UE4Game-StructBox-IOS-Shipping-Static.a" Type="StaticLibrary" />
+    <BuildProduct Path="$(EngineDir)/Plugins/Media/AvfMedia/Binaries/IOS/UE4Game-AvfMedia-IOS-Shipping-Static.a" Type="StaticLibrary" />
+    <BuildProduct Path="$(EngineDir)/Plugins/Messaging/UdpMessaging/Binaries/IOS/UE4Game-UdpMessaging-IOS-Shipping-Static.a" Type="StaticLibrary" />
+    <BuildProduct Path="$(EngineDir)/Plugins/Runtime/Advertising/IOSTapJoy/Binaries/IOS/UE4Game-IOSTapJoy-IOS-Shipping-Static.a" Type="StaticLibrary" />
+    <BuildProduct Path="$(EngineDir)/Plugins/Runtime/Analytics/AnalyticsBlueprintLibrary/Binaries/IOS/UE4Game-AnalyticsBlueprintLibrary-IOS-Shipping-Static.a" Type="StaticLibrary" />
+    <BuildProduct Path="$(EngineDir)/Plugins/Runtime/Analytics/Apsalar/IOSApsalar/Binaries/IOS/NotForLicensees/UE4Game-IOSApsalar-IOS-Shipping-Static.a" Type="StaticLibrary" />
+    <BuildProduct Path="$(EngineDir)/Plugins/Runtime/Analytics/FileLogging/Binaries/IOS/UE4Game-FileLogging-IOS-Shipping-Static.a" Type="StaticLibrary" />
+    <BuildProduct Path="$(EngineDir)/Plugins/Runtime/AppleMoviePlayer/Binaries/IOS/UE4Game-AppleMoviePlayer-IOS-Shipping-Static.a" Type="StaticLibrary" />
+    <BuildProduct Path="$(EngineDir)/Plugins/Runtime/CableComponent/Binaries/IOS/UE4Game-CableComponent-IOS-Shipping-Static.a" Type="StaticLibrary" />
+    <BuildProduct Path="$(EngineDir)/Plugins/Runtime/CustomMeshComponent/Binaries/IOS/UE4Game-CustomMeshComponent-IOS-Shipping-Static.a" Type="StaticLibrary" />
+    <BuildProduct Path="$(EngineDir)/Plugins/Runtime/ExampleDeviceProfileSelector/Binaries/IOS/UE4Game-ExampleDeviceProfileSelector-IOS-Shipping-Static.a" Type="StaticLibrary" />
+    <BuildProduct Path="$(EngineDir)/Plugins/Runtime/IOSDeviceProfileSelector/Binaries/IOS/UE4Game-IOSDeviceProfileSelector-IOS-Shipping-Static.a" Type="StaticLibrary" />
+    <BuildProduct Path="$(EngineDir)/Plugins/Runtime/SoundVisualizations/Binaries/IOS/UE4Game-SoundVisualizations-IOS-Shipping-Static.a" Type="StaticLibrary" />
+    <BuildProduct Path="$(EngineDir)/Plugins/ScriptPlugin/Binaries/IOS/UE4Game-ScriptPlugin-IOS-Shipping-Static.a" Type="StaticLibrary" />
+    <BuildProduct Path="$(EngineDir)/Plugins/Slate/SlateRemote/Binaries/IOS/UE4Game-SlateRemote-IOS-Shipping-Static.a" Type="StaticLibrary" />
+    <BuildProduct Path="$(EngineDir)/Binaries/IOS/UE4Game-Launch-IOS-Shipping.a" Type="StaticLibrary" />
+    <BuildProduct Path="$(EngineDir)/Binaries/IOS/UE4Game-Core-IOS-Shipping.a" Type="StaticLibrary" />
+    <BuildProduct Path="$(EngineDir)/Binaries/IOS/UE4Game-CoreUObject-IOS-Shipping.a" Type="StaticLibrary" />
+    <BuildProduct Path="$(EngineDir)/Binaries/IOS/UE4Game-Projects-IOS-Shipping.a" Type="StaticLibrary" />
+    <BuildProduct Path="$(EngineDir)/Binaries/IOS/UE4Game-Json-IOS-Shipping.a" Type="StaticLibrary" />
+    <BuildProduct Path="$(EngineDir)/Binaries/IOS/UE4Game-Engine-IOS-Shipping.a" Type="StaticLibrary" />
+    <BuildProduct Path="$(EngineDir)/Binaries/IOS/UE4Game-AppFramework-IOS-Shipping.a" Type="StaticLibrary" />
+    <BuildProduct Path="$(EngineDir)/Binaries/IOS/UE4Game-InputCore-IOS-Shipping.a" Type="StaticLibrary" />
+    <BuildProduct Path="$(EngineDir)/Binaries/IOS/UE4Game-SlateReflector-IOS-Shipping.a" Type="StaticLibrary" />
+    <BuildProduct Path="$(EngineDir)/Binaries/IOS/UE4Game-Slate-IOS-Shipping.a" Type="StaticLibrary" />
+    <BuildProduct Path="$(EngineDir)/Binaries/IOS/UE4Game-SlateCore-IOS-Shipping.a" Type="StaticLibrary" />
+    <BuildProduct Path="$(EngineDir)/Binaries/IOS/UE4Game-Networking-IOS-Shipping.a" Type="StaticLibrary" />
+    <BuildProduct Path="$(EngineDir)/Binaries/IOS/UE4Game-Sockets-IOS-Shipping.a" Type="StaticLibrary" />
+    <BuildProduct Path="$(EngineDir)/Binaries/IOS/UE4Game-SandboxFile-IOS-Shipping.a" Type="StaticLibrary" />
+    <BuildProduct Path="$(EngineDir)/Binaries/IOS/UE4Game-VectorVM-IOS-Shipping.a" Type="StaticLibrary" />
+    <BuildProduct Path="$(EngineDir)/Binaries/IOS/UE4Game-Landscape-IOS-Shipping.a" Type="StaticLibrary" />
+    <BuildProduct Path="$(EngineDir)/Binaries/IOS/UE4Game-RenderCore-IOS-Shipping.a" Type="StaticLibrary" />
+    <BuildProduct Path="$(EngineDir)/Binaries/IOS/UE4Game-RHI-IOS-Shipping.a" Type="StaticLibrary" />
+    <BuildProduct Path="$(EngineDir)/Binaries/IOS/UE4Game-NullDrv-IOS-Shipping.a" Type="StaticLibrary" />
+    <BuildProduct Path="$(EngineDir)/Binaries/IOS/UE4Game-ShaderCore-IOS-Shipping.a" Type="StaticLibrary" />
+    <BuildProduct Path="$(EngineDir)/Binaries/IOS/UE4Game-Renderer-IOS-Shipping.a" Type="StaticLibrary" />
+    <BuildProduct Path="$(EngineDir)/Binaries/IOS/UE4Game-UtilityShaders-IOS-Shipping.a" Type="StaticLibrary" />
+    <BuildProduct Path="$(EngineDir)/Binaries/IOS/UE4Game-HeadMountedDisplay-IOS-Shipping.a" Type="StaticLibrary" />
+    <BuildProduct Path="$(EngineDir)/Binaries/IOS/UE4Game-Foliage-IOS-Shipping.a" Type="StaticLibrary" />
+    <BuildProduct Path="$(EngineDir)/Binaries/IOS/UE4Game-UMG-IOS-Shipping.a" Type="StaticLibrary" />
+    <BuildProduct Path="$(EngineDir)/Binaries/IOS/UE4Game-MovieSceneCore-IOS-Shipping.a" Type="StaticLibrary" />
+    <BuildProduct Path="$(EngineDir)/Binaries/IOS/UE4Game-HTTP-IOS-Shipping.a" Type="StaticLibrary" />
+    <BuildProduct Path="$(EngineDir)/Binaries/IOS/UE4Game-ImageWrapper-IOS-Shipping.a" Type="StaticLibrary" />
+    <BuildProduct Path="$(EngineDir)/Binaries/IOS/UE4Game-SlateRHIRenderer-IOS-Shipping.a" Type="StaticLibrary" />
+    <BuildProduct Path="$(EngineDir)/Binaries/IOS/UE4Game-Navmesh-IOS-Shipping.a" Type="StaticLibrary" />
+    <BuildProduct Path="$(EngineDir)/Binaries/IOS/UE4Game-Messaging-IOS-Shipping.a" Type="StaticLibrary" />
+    <BuildProduct Path="$(EngineDir)/Binaries/IOS/UE4Game-AssetRegistry-IOS-Shipping.a" Type="StaticLibrary" />
+    <BuildProduct Path="$(EngineDir)/Binaries/IOS/UE4Game-EngineMessages-IOS-Shipping.a" Type="StaticLibrary" />
+    <BuildProduct Path="$(EngineDir)/Binaries/IOS/UE4Game-EngineSettings-IOS-Shipping.a" Type="StaticLibrary" />
+    <BuildProduct Path="$(EngineDir)/Binaries/IOS/UE4Game-SynthBenchmark-IOS-Shipping.a" Type="StaticLibrary" />
+    <BuildProduct Path="$(EngineDir)/Binaries/IOS/UE4Game-AIModule-IOS-Shipping.a" Type="StaticLibrary" />
+    <BuildProduct Path="$(EngineDir)/Binaries/IOS/UE4Game-GameplayTags-IOS-Shipping.a" Type="StaticLibrary" />
+    <BuildProduct Path="$(EngineDir)/Binaries/IOS/UE4Game-DatabaseSupport-IOS-Shipping.a" Type="StaticLibrary" />
+    <BuildProduct Path="$(EngineDir)/Binaries/IOS/UE4Game-MovieSceneCoreTypes-IOS-Shipping.a" Type="StaticLibrary" />
+    <BuildProduct Path="$(EngineDir)/Binaries/IOS/UE4Game-StreamingPauseRendering-IOS-Shipping.a" Type="StaticLibrary" />
+    <BuildProduct Path="$(EngineDir)/Binaries/IOS/UE4Game-MoviePlayer-IOS-Shipping.a" Type="StaticLibrary" />
+    <BuildProduct Path="$(EngineDir)/Binaries/IOS/UE4Game-Analytics-IOS-Shipping.a" Type="StaticLibrary" />
+    <BuildProduct Path="$(EngineDir)/Binaries/IOS/UE4Game-AnalyticsET-IOS-Shipping.a" Type="StaticLibrary" />
+    <BuildProduct Path="$(EngineDir)/Binaries/IOS/UE4Game-NetworkReplayStreaming-IOS-Shipping.a" Type="StaticLibrary" />
+    <BuildProduct Path="$(EngineDir)/Binaries/IOS/UE4Game-NullNetworkReplayStreaming-IOS-Shipping.a" Type="StaticLibrary" />
+    <BuildProduct Path="$(EngineDir)/Binaries/IOS/UE4Game-HttpNetworkReplayStreaming-IOS-Shipping.a" Type="StaticLibrary" />
+    <BuildProduct Path="$(EngineDir)/Binaries/IOS/UE4Game-OnlineSubsystem-IOS-Shipping.a" Type="StaticLibrary" />
+    <BuildProduct Path="$(EngineDir)/Binaries/IOS/UE4Game-ImageCore-IOS-Shipping.a" Type="StaticLibrary" />
+    <BuildProduct Path="$(EngineDir)/Binaries/IOS/UE4Game-OnlineSubsystemUtils-IOS-Shipping.a" Type="StaticLibrary" />
+    <BuildProduct Path="$(EngineDir)/Binaries/IOS/UE4Game-Voice-IOS-Shipping.a" Type="StaticLibrary" />
+    <BuildProduct Path="$(EngineDir)/Binaries/IOS/UE4Game-Advertising-IOS-Shipping.a" Type="StaticLibrary" />
+    <BuildProduct Path="$(EngineDir)/Binaries/IOS/UE4Game-CrashTracker-IOS-Shipping.a" Type="StaticLibrary" />
+    <BuildProduct Path="$(EngineDir)/Binaries/IOS/UE4Game-GameLiveStreaming-IOS-Shipping.a" Type="StaticLibrary" />
+    <BuildProduct Path="$(EngineDir)/Binaries/IOS/UE4Game-MediaAssets-IOS-Shipping.a" Type="StaticLibrary" />
+    <BuildProduct Path="$(EngineDir)/Binaries/IOS/UE4Game-Media-IOS-Shipping.a" Type="StaticLibrary" />
+    <BuildProduct Path="$(EngineDir)/Binaries/IOS/UE4Game-PakFile-IOS-Shipping.a" Type="StaticLibrary" />
+    <BuildProduct Path="$(EngineDir)/Binaries/IOS/UE4Game-Serialization-IOS-Shipping.a" Type="StaticLibrary" />
+    <BuildProduct Path="$(EngineDir)/Binaries/IOS/UE4Game-OpenGLDrv-IOS-Shipping.a" Type="StaticLibrary" />
+    <BuildProduct Path="$(EngineDir)/Binaries/IOS/UE4Game-LaunchDaemonMessages-IOS-Shipping.a" Type="StaticLibrary" />
+    <BuildProduct Path="$(EngineDir)/Binaries/IOS/UE4Game-SessionServices-IOS-Shipping.a" Type="StaticLibrary" />
+    <BuildProduct Path="$(EngineDir)/Binaries/IOS/UE4Game-SessionMessages-IOS-Shipping.a" Type="StaticLibrary" />
+    <BuildProduct Path="$(EngineDir)/Binaries/IOS/IOS/UE4Game-IOSAudio-IOS-Shipping.a" Type="StaticLibrary" />
+    <BuildProduct Path="$(EngineDir)/Binaries/IOS/IOS/UE4Game-IOSRuntimeSettings-IOS-Shipping.a" Type="StaticLibrary" />
+    <BuildProduct Path="$(EngineDir)/Binaries/IOS/UE4Game-InputDevice-IOS-Shipping.a" Type="StaticLibrary" />
+    <BuildProduct Path="$(EngineDir)/Binaries/IOS/UE4Game-GameMenuBuilder-IOS-Shipping.a" Type="StaticLibrary" />
+    <BuildProduct Path="$(EngineDir)/Binaries/IOS/UE4Game-XmlParser-IOS-Shipping.a" Type="StaticLibrary" />
+    <BuildProduct Path="$(EngineDir)/Binaries/IOS/UE4Game-OnlineSubsystemFacebook-IOS-Shipping.a" Type="StaticLibrary" />
+    <BuildProduct Path="$(EngineDir)/Binaries/IOS/UE4Game-OnlineSubsystemIOS-IOS-Shipping.a" Type="StaticLibrary" />
+    <BuildProduct Path="$(EngineDir)/Binaries/IOS/UE4Game-IOSAdvertising-IOS-Shipping.a" Type="StaticLibrary" />
+  </BuildProducts>
+  <RuntimeDependencies />
+</BuildReceipt>
\ No newline at end of file
diff --git a/Engine/Build/Receipts/UE4Game-IOS-Test.target.xml b/Engine/Build/Receipts/UE4Game-IOS-Test.target.xml
new file mode 100644
--- /dev/null
+++ b/Engine/Build/Receipts/UE4Game-IOS-Test.target.xml
@@ -0,0 +1,18 @@
+<?xml version="1.0" encoding="utf-8"?>
+<BuildReceipt xmlns:xsi="http://www.w3.org/2001/XMLSchema-instance" xmlns:xsd="http://www.w3.org/2001/XMLSchema">
+  <BuildProducts>
+    <BuildProduct Path="$(EngineDir)/Binaries/IOS/UE4Game-IOS-Test" Type="Executable" />
+    <BuildProduct Path="$(EngineDir)/Binaries/IOS/UE4Game-IOS-Test.stub" Type="Executable" />
+    <BuildProduct Path="$(EngineDir)/Plugins/2D/Paper2D/Binaries/IOS/UE4Game-Paper2D-IOS-Test-Static.a" Type="StaticLibrary" />
+    <BuildProduct Path="$(EngineDir)/Plugins/Experimental/CharacterAI/Binaries/IOS/UE4Game-CharacterAI-IOS-Test-Static.a" Type="StaticLibrary" />
+    <BuildProduct Path="$(EngineDir)/Plugins/Media/AvfMedia/Binaries/IOS/UE4Game-AvfMedia-IOS-Test-Static.a" Type="StaticLibrary" />
+    <BuildProduct Path="$(EngineDir)/Plugins/Messaging/UdpMessaging/Binaries/IOS/UE4Game-UdpMessaging-IOS-Test-Static.a" Type="StaticLibrary" />
+    <BuildProduct Path="$(EngineDir)/Plugins/Runtime/AppleMoviePlayer/Binaries/IOS/UE4Game-AppleMoviePlayer-IOS-Test-Static.a" Type="StaticLibrary" />
+    <BuildProduct Path="$(EngineDir)/Plugins/Runtime/CableComponent/Binaries/IOS/UE4Game-CableComponent-IOS-Test-Static.a" Type="StaticLibrary" />
+    <BuildProduct Path="$(EngineDir)/Plugins/Runtime/CustomMeshComponent/Binaries/IOS/UE4Game-CustomMeshComponent-IOS-Test-Static.a" Type="StaticLibrary" />
+    <BuildProduct Path="$(EngineDir)/Plugins/Runtime/ExampleDeviceProfileSelector/Binaries/IOS/UE4Game-ExampleDeviceProfileSelector-IOS-Test-Static.a" Type="StaticLibrary" />
+    <BuildProduct Path="$(EngineDir)/Plugins/Runtime/IOSDeviceProfileSelector/Binaries/IOS/UE4Game-IOSDeviceProfileSelector-IOS-Test-Static.a" Type="StaticLibrary" />
+    <BuildProduct Path="$(EngineDir)/Plugins/Slate/SlateRemote/Binaries/IOS/UE4Game-SlateRemote-IOS-Test-Static.a" Type="StaticLibrary" />
+  </BuildProducts>
+  <RuntimeDependencies />
+</BuildReceipt>
\ No newline at end of file
diff --git a/Engine/Build/Receipts/UE4Game-Linux-Developmentx86_64-unknown-linux-gnu.target.xml b/Engine/Build/Receipts/UE4Game-Linux-Developmentx86_64-unknown-linux-gnu.target.xml
new file mode 100644
--- /dev/null
+++ b/Engine/Build/Receipts/UE4Game-Linux-Developmentx86_64-unknown-linux-gnu.target.xml
@@ -0,0 +1,14 @@
+<?xml version="1.0" encoding="utf-8"?>
+<BuildReceipt xmlns:xsd="http://www.w3.org/2001/XMLSchema" xmlns:xsi="http://www.w3.org/2001/XMLSchema-instance">
+  <BuildProducts>
+    <BuildProduct Path="$(EngineDir)\Binaries\Linux\UE4Game" Type="Executable" />
+    <BuildProduct Path="$(EngineDir)\Plugins\2D\Paper2D\Binaries\Linux\libUE4Game-Paper2D-Static.a" Type="StaticLibrary" />
+    <BuildProduct Path="$(EngineDir)\Plugins\Experimental\CharacterAI\Binaries\Linux\libUE4Game-CharacterAI-Static.a" Type="StaticLibrary" />
+    <BuildProduct Path="$(EngineDir)\Plugins\Messaging\UdpMessaging\Binaries\Linux\libUE4Game-UdpMessaging-Static.a" Type="StaticLibrary" />
+    <BuildProduct Path="$(EngineDir)\Plugins\Runtime\CableComponent\Binaries\Linux\libUE4Game-CableComponent-Static.a" Type="StaticLibrary" />
+    <BuildProduct Path="$(EngineDir)\Plugins\Runtime\CustomMeshComponent\Binaries\Linux\libUE4Game-CustomMeshComponent-Static.a" Type="StaticLibrary" />
+    <BuildProduct Path="$(EngineDir)\Plugins\Runtime\ExampleDeviceProfileSelector\Binaries\Linux\libUE4Game-ExampleDeviceProfileSelector-Static.a" Type="StaticLibrary" />
+    <BuildProduct Path="$(EngineDir)\Plugins\Slate\SlateRemote\Binaries\Linux\libUE4Game-SlateRemote-Static.a" Type="StaticLibrary" />
+  </BuildProducts>
+  <RuntimeDependencies />
+</BuildReceipt>
\ No newline at end of file
diff --git a/Engine/Build/Receipts/UE4Game-Linux-Shippingx86_64-unknown-linux-gnu.target.xml b/Engine/Build/Receipts/UE4Game-Linux-Shippingx86_64-unknown-linux-gnu.target.xml
new file mode 100644
--- /dev/null
+++ b/Engine/Build/Receipts/UE4Game-Linux-Shippingx86_64-unknown-linux-gnu.target.xml
@@ -0,0 +1,14 @@
+<?xml version="1.0" encoding="utf-8"?>
+<BuildReceipt xmlns:xsd="http://www.w3.org/2001/XMLSchema" xmlns:xsi="http://www.w3.org/2001/XMLSchema-instance">
+  <BuildProducts>
+    <BuildProduct Path="$(EngineDir)\Binaries\Linux\UE4Game-Linux-Shipping" Type="Executable" />
+    <BuildProduct Path="$(EngineDir)\Plugins\2D\Paper2D\Binaries\Linux\libUE4Game-Paper2D-Linux-Shipping-Static.a" Type="StaticLibrary" />
+    <BuildProduct Path="$(EngineDir)\Plugins\Experimental\CharacterAI\Binaries\Linux\libUE4Game-CharacterAI-Linux-Shipping-Static.a" Type="StaticLibrary" />
+    <BuildProduct Path="$(EngineDir)\Plugins\Messaging\UdpMessaging\Binaries\Linux\libUE4Game-UdpMessaging-Linux-Shipping-Static.a" Type="StaticLibrary" />
+    <BuildProduct Path="$(EngineDir)\Plugins\Runtime\CableComponent\Binaries\Linux\libUE4Game-CableComponent-Linux-Shipping-Static.a" Type="StaticLibrary" />
+    <BuildProduct Path="$(EngineDir)\Plugins\Runtime\CustomMeshComponent\Binaries\Linux\libUE4Game-CustomMeshComponent-Linux-Shipping-Static.a" Type="StaticLibrary" />
+    <BuildProduct Path="$(EngineDir)\Plugins\Runtime\ExampleDeviceProfileSelector\Binaries\Linux\libUE4Game-ExampleDeviceProfileSelector-Linux-Shipping-Static.a" Type="StaticLibrary" />
+    <BuildProduct Path="$(EngineDir)\Plugins\Slate\SlateRemote\Binaries\Linux\libUE4Game-SlateRemote-Linux-Shipping-Static.a" Type="StaticLibrary" />
+  </BuildProducts>
+  <RuntimeDependencies />
+</BuildReceipt>
\ No newline at end of file
diff --git a/Engine/Build/Receipts/UE4Game-Linux-Testx86_64-unknown-linux-gnu.target.xml b/Engine/Build/Receipts/UE4Game-Linux-Testx86_64-unknown-linux-gnu.target.xml
new file mode 100644
--- /dev/null
+++ b/Engine/Build/Receipts/UE4Game-Linux-Testx86_64-unknown-linux-gnu.target.xml
@@ -0,0 +1,14 @@
+<?xml version="1.0" encoding="utf-8"?>
+<BuildReceipt xmlns:xsd="http://www.w3.org/2001/XMLSchema" xmlns:xsi="http://www.w3.org/2001/XMLSchema-instance">
+  <BuildProducts>
+    <BuildProduct Path="$(EngineDir)\Binaries\Linux\UE4Game-Linux-Test" Type="Executable" />
+    <BuildProduct Path="$(EngineDir)\Plugins\2D\Paper2D\Binaries\Linux\libUE4Game-Paper2D-Linux-Test-Static.a" Type="StaticLibrary" />
+    <BuildProduct Path="$(EngineDir)\Plugins\Experimental\CharacterAI\Binaries\Linux\libUE4Game-CharacterAI-Linux-Test-Static.a" Type="StaticLibrary" />
+    <BuildProduct Path="$(EngineDir)\Plugins\Messaging\UdpMessaging\Binaries\Linux\libUE4Game-UdpMessaging-Linux-Test-Static.a" Type="StaticLibrary" />
+    <BuildProduct Path="$(EngineDir)\Plugins\Runtime\CableComponent\Binaries\Linux\libUE4Game-CableComponent-Linux-Test-Static.a" Type="StaticLibrary" />
+    <BuildProduct Path="$(EngineDir)\Plugins\Runtime\CustomMeshComponent\Binaries\Linux\libUE4Game-CustomMeshComponent-Linux-Test-Static.a" Type="StaticLibrary" />
+    <BuildProduct Path="$(EngineDir)\Plugins\Runtime\ExampleDeviceProfileSelector\Binaries\Linux\libUE4Game-ExampleDeviceProfileSelector-Linux-Test-Static.a" Type="StaticLibrary" />
+    <BuildProduct Path="$(EngineDir)\Plugins\Slate\SlateRemote\Binaries\Linux\libUE4Game-SlateRemote-Linux-Test-Static.a" Type="StaticLibrary" />
+  </BuildProducts>
+  <RuntimeDependencies />
+</BuildReceipt>
\ No newline at end of file
diff --git a/Engine/Build/Receipts/UE4Game-Mac-Development.target.xml b/Engine/Build/Receipts/UE4Game-Mac-Development.target.xml
new file mode 100644
--- /dev/null
+++ b/Engine/Build/Receipts/UE4Game-Mac-Development.target.xml
@@ -0,0 +1,425 @@
+<?xml version="1.0" encoding="utf-8"?>
+<BuildReceipt xmlns:xsi="http://www.w3.org/2001/XMLSchema-instance" xmlns:xsd="http://www.w3.org/2001/XMLSchema">
+  <BuildProducts>
+    <BuildProduct Path="$(EngineDir)/Binaries/Mac/UE4Game.app/Contents/MacOS/UE4Game" Type="Executable" />
+    <BuildProduct Path="$(EngineDir)/Binaries/Mac/UE4Game.app/Contents/MacOS/libogg.dylib" Type="DynamicLibrary" />
+    <BuildProduct Path="$(EngineDir)/Binaries/Mac/UE4Game.app/Contents/MacOS/libvorbis.dylib" Type="DynamicLibrary" />
+    <BuildProduct Path="$(EngineDir)/Binaries/Mac/UE4Game.app/Contents/MacOS/libplugin_carbon_interpose.dylib" Type="DynamicLibrary" />
+    <BuildProduct Path="$(EngineDir)/Binaries/Mac/UE4Game.app/Contents/Resources/RadioEffectUnit.component/Contents/Info.plist" Type="RequiredResource" />
+    <BuildProduct Path="$(EngineDir)/Binaries/Mac/UE4Game.app/Contents/Resources/RadioEffectUnit.component/Contents/MacOS/RadioEffectUnit" Type="RequiredResource" />
+    <BuildProduct Path="$(EngineDir)/Binaries/Mac/UE4Game.app/Contents/Resources/RadioEffectUnit.component/Contents/Resources/English.lproj/Localizable.strings" Type="RequiredResource" />
+    <BuildProduct Path="$(EngineDir)/Binaries/Mac/UE4Game.app/Contents/Resources/am.lproj/locale.pak" Type="RequiredResource" />
+    <BuildProduct Path="$(EngineDir)/Binaries/Mac/UE4Game.app/Contents/Resources/ar.lproj/locale.pak" Type="RequiredResource" />
+    <BuildProduct Path="$(EngineDir)/Binaries/Mac/UE4Game.app/Contents/Resources/bg.lproj/locale.pak" Type="RequiredResource" />
+    <BuildProduct Path="$(EngineDir)/Binaries/Mac/UE4Game.app/Contents/Resources/bn.lproj/locale.pak" Type="RequiredResource" />
+    <BuildProduct Path="$(EngineDir)/Binaries/Mac/UE4Game.app/Contents/Resources/ca.lproj/locale.pak" Type="RequiredResource" />
+    <BuildProduct Path="$(EngineDir)/Binaries/Mac/UE4Game.app/Contents/Resources/cs.lproj/locale.pak" Type="RequiredResource" />
+    <BuildProduct Path="$(EngineDir)/Binaries/Mac/UE4Game.app/Contents/Resources/da.lproj/locale.pak" Type="RequiredResource" />
+    <BuildProduct Path="$(EngineDir)/Binaries/Mac/UE4Game.app/Contents/Resources/de.lproj/locale.pak" Type="RequiredResource" />
+    <BuildProduct Path="$(EngineDir)/Binaries/Mac/UE4Game.app/Contents/Resources/el.lproj/locale.pak" Type="RequiredResource" />
+    <BuildProduct Path="$(EngineDir)/Binaries/Mac/UE4Game.app/Contents/Resources/en.lproj/locale.pak" Type="RequiredResource" />
+    <BuildProduct Path="$(EngineDir)/Binaries/Mac/UE4Game.app/Contents/Resources/en_GB.lproj/locale.pak" Type="RequiredResource" />
+    <BuildProduct Path="$(EngineDir)/Binaries/Mac/UE4Game.app/Contents/Resources/es.lproj/locale.pak" Type="RequiredResource" />
+    <BuildProduct Path="$(EngineDir)/Binaries/Mac/UE4Game.app/Contents/Resources/es_419.lproj/locale.pak" Type="RequiredResource" />
+    <BuildProduct Path="$(EngineDir)/Binaries/Mac/UE4Game.app/Contents/Resources/et.lproj/locale.pak" Type="RequiredResource" />
+    <BuildProduct Path="$(EngineDir)/Binaries/Mac/UE4Game.app/Contents/Resources/fa.lproj/locale.pak" Type="RequiredResource" />
+    <BuildProduct Path="$(EngineDir)/Binaries/Mac/UE4Game.app/Contents/Resources/fi.lproj/locale.pak" Type="RequiredResource" />
+    <BuildProduct Path="$(EngineDir)/Binaries/Mac/UE4Game.app/Contents/Resources/fil.lproj/locale.pak" Type="RequiredResource" />
+    <BuildProduct Path="$(EngineDir)/Binaries/Mac/UE4Game.app/Contents/Resources/fr.lproj/locale.pak" Type="RequiredResource" />
+    <BuildProduct Path="$(EngineDir)/Binaries/Mac/UE4Game.app/Contents/Resources/gu.lproj/locale.pak" Type="RequiredResource" />
+    <BuildProduct Path="$(EngineDir)/Binaries/Mac/UE4Game.app/Contents/Resources/he.lproj/locale.pak" Type="RequiredResource" />
+    <BuildProduct Path="$(EngineDir)/Binaries/Mac/UE4Game.app/Contents/Resources/hi.lproj/locale.pak" Type="RequiredResource" />
+    <BuildProduct Path="$(EngineDir)/Binaries/Mac/UE4Game.app/Contents/Resources/hr.lproj/locale.pak" Type="RequiredResource" />
+    <BuildProduct Path="$(EngineDir)/Binaries/Mac/UE4Game.app/Contents/Resources/hu.lproj/locale.pak" Type="RequiredResource" />
+    <BuildProduct Path="$(EngineDir)/Binaries/Mac/UE4Game.app/Contents/Resources/id.lproj/locale.pak" Type="RequiredResource" />
+    <BuildProduct Path="$(EngineDir)/Binaries/Mac/UE4Game.app/Contents/Resources/it.lproj/locale.pak" Type="RequiredResource" />
+    <BuildProduct Path="$(EngineDir)/Binaries/Mac/UE4Game.app/Contents/Resources/ja.lproj/locale.pak" Type="RequiredResource" />
+    <BuildProduct Path="$(EngineDir)/Binaries/Mac/UE4Game.app/Contents/Resources/kn.lproj/locale.pak" Type="RequiredResource" />
+    <BuildProduct Path="$(EngineDir)/Binaries/Mac/UE4Game.app/Contents/Resources/ko.lproj/locale.pak" Type="RequiredResource" />
+    <BuildProduct Path="$(EngineDir)/Binaries/Mac/UE4Game.app/Contents/Resources/lt.lproj/locale.pak" Type="RequiredResource" />
+    <BuildProduct Path="$(EngineDir)/Binaries/Mac/UE4Game.app/Contents/Resources/lv.lproj/locale.pak" Type="RequiredResource" />
+    <BuildProduct Path="$(EngineDir)/Binaries/Mac/UE4Game.app/Contents/Resources/ml.lproj/locale.pak" Type="RequiredResource" />
+    <BuildProduct Path="$(EngineDir)/Binaries/Mac/UE4Game.app/Contents/Resources/mr.lproj/locale.pak" Type="RequiredResource" />
+    <BuildProduct Path="$(EngineDir)/Binaries/Mac/UE4Game.app/Contents/Resources/ms.lproj/locale.pak" Type="RequiredResource" />
+    <BuildProduct Path="$(EngineDir)/Binaries/Mac/UE4Game.app/Contents/Resources/nb.lproj/locale.pak" Type="RequiredResource" />
+    <BuildProduct Path="$(EngineDir)/Binaries/Mac/UE4Game.app/Contents/Resources/nl.lproj/locale.pak" Type="RequiredResource" />
+    <BuildProduct Path="$(EngineDir)/Binaries/Mac/UE4Game.app/Contents/Resources/pl.lproj/locale.pak" Type="RequiredResource" />
+    <BuildProduct Path="$(EngineDir)/Binaries/Mac/UE4Game.app/Contents/Resources/pt_BR.lproj/locale.pak" Type="RequiredResource" />
+    <BuildProduct Path="$(EngineDir)/Binaries/Mac/UE4Game.app/Contents/Resources/pt_PT.lproj/locale.pak" Type="RequiredResource" />
+    <BuildProduct Path="$(EngineDir)/Binaries/Mac/UE4Game.app/Contents/Resources/ro.lproj/locale.pak" Type="RequiredResource" />
+    <BuildProduct Path="$(EngineDir)/Binaries/Mac/UE4Game.app/Contents/Resources/ru.lproj/locale.pak" Type="RequiredResource" />
+    <BuildProduct Path="$(EngineDir)/Binaries/Mac/UE4Game.app/Contents/Resources/sk.lproj/locale.pak" Type="RequiredResource" />
+    <BuildProduct Path="$(EngineDir)/Binaries/Mac/UE4Game.app/Contents/Resources/sl.lproj/locale.pak" Type="RequiredResource" />
+    <BuildProduct Path="$(EngineDir)/Binaries/Mac/UE4Game.app/Contents/Resources/sr.lproj/locale.pak" Type="RequiredResource" />
+    <BuildProduct Path="$(EngineDir)/Binaries/Mac/UE4Game.app/Contents/Resources/sv.lproj/locale.pak" Type="RequiredResource" />
+    <BuildProduct Path="$(EngineDir)/Binaries/Mac/UE4Game.app/Contents/Resources/sw.lproj/locale.pak" Type="RequiredResource" />
+    <BuildProduct Path="$(EngineDir)/Binaries/Mac/UE4Game.app/Contents/Resources/ta.lproj/locale.pak" Type="RequiredResource" />
+    <BuildProduct Path="$(EngineDir)/Binaries/Mac/UE4Game.app/Contents/Resources/te.lproj/locale.pak" Type="RequiredResource" />
+    <BuildProduct Path="$(EngineDir)/Binaries/Mac/UE4Game.app/Contents/Resources/th.lproj/locale.pak" Type="RequiredResource" />
+    <BuildProduct Path="$(EngineDir)/Binaries/Mac/UE4Game.app/Contents/Resources/tr.lproj/locale.pak" Type="RequiredResource" />
+    <BuildProduct Path="$(EngineDir)/Binaries/Mac/UE4Game.app/Contents/Resources/uk.lproj/locale.pak" Type="RequiredResource" />
+    <BuildProduct Path="$(EngineDir)/Binaries/Mac/UE4Game.app/Contents/Resources/vi.lproj/locale.pak" Type="RequiredResource" />
+    <BuildProduct Path="$(EngineDir)/Binaries/Mac/UE4Game.app/Contents/Resources/zh_CN.lproj/locale.pak" Type="RequiredResource" />
+    <BuildProduct Path="$(EngineDir)/Binaries/Mac/UE4Game.app/Contents/Resources/zh_TW.lproj/locale.pak" Type="RequiredResource" />
+    <BuildProduct Path="$(EngineDir)/Binaries/Mac/UE4Game.app/Contents/Info.plist" Type="RequiredResource" />
+    <BuildProduct Path="$(EngineDir)/Binaries/Mac/UE4Game.app/Contents/PkgInfo" Type="RequiredResource" />
+    <BuildProduct Path="$(EngineDir)/Binaries/Mac/UE4Game.app/Contents/Resources/UE4.icns" Type="RequiredResource" />
+    <BuildProduct Path="$(EngineDir)/Plugins/2D/Paper2D/Binaries/Mac/UE4Game-Paper2D-Static.a" Type="StaticLibrary" />
+    <BuildProduct Path="$(EngineDir)/Binaries/Mac/UE4Game.app/Contents/MacOS/libogg.dylib" Type="DynamicLibrary" />
+    <BuildProduct Path="$(EngineDir)/Binaries/Mac/UE4Game.app/Contents/MacOS/libvorbis.dylib" Type="DynamicLibrary" />
+    <BuildProduct Path="$(EngineDir)/Plugins/Experimental/CharacterAI/Binaries/Mac/UE4Game-CharacterAI-Static.a" Type="StaticLibrary" />
+    <BuildProduct Path="$(EngineDir)/Binaries/Mac/UE4Game.app/Contents/MacOS/libogg.dylib" Type="DynamicLibrary" />
+    <BuildProduct Path="$(EngineDir)/Binaries/Mac/UE4Game.app/Contents/MacOS/libvorbis.dylib" Type="DynamicLibrary" />
+    <BuildProduct Path="$(EngineDir)/Plugins/Experimental/SimpleHMD/Binaries/Mac/UE4Game-SimpleHMD-Static.a" Type="StaticLibrary" />
+    <BuildProduct Path="$(EngineDir)/Binaries/Mac/UE4Game.app/Contents/MacOS/libogg.dylib" Type="DynamicLibrary" />
+    <BuildProduct Path="$(EngineDir)/Binaries/Mac/UE4Game.app/Contents/MacOS/libvorbis.dylib" Type="DynamicLibrary" />
+    <BuildProduct Path="$(EngineDir)/Plugins/Experimental/StructBox/Binaries/Mac/UE4Game-StructBox-Static.a" Type="StaticLibrary" />
+    <BuildProduct Path="$(EngineDir)/Binaries/Mac/UE4Game.app/Contents/MacOS/libogg.dylib" Type="DynamicLibrary" />
+    <BuildProduct Path="$(EngineDir)/Binaries/Mac/UE4Game.app/Contents/MacOS/libvorbis.dylib" Type="DynamicLibrary" />
+    <BuildProduct Path="$(EngineDir)/Plugins/Media/AvfMedia/Binaries/Mac/UE4Game-AvfMedia-Static.a" Type="StaticLibrary" />
+    <BuildProduct Path="$(EngineDir)/Plugins/Messaging/UdpMessaging/Binaries/Mac/UE4Game-UdpMessaging-Static.a" Type="StaticLibrary" />
+    <BuildProduct Path="$(EngineDir)/Plugins/Runtime/Analytics/AnalyticsBlueprintLibrary/Binaries/Mac/UE4Game-AnalyticsBlueprintLibrary-Static.a" Type="StaticLibrary" />
+    <BuildProduct Path="$(EngineDir)/Binaries/Mac/UE4Game.app/Contents/MacOS/libogg.dylib" Type="DynamicLibrary" />
+    <BuildProduct Path="$(EngineDir)/Binaries/Mac/UE4Game.app/Contents/MacOS/libvorbis.dylib" Type="DynamicLibrary" />
+    <BuildProduct Path="$(EngineDir)/Plugins/Runtime/Analytics/FileLogging/Binaries/Mac/UE4Game-FileLogging-Static.a" Type="StaticLibrary" />
+    <BuildProduct Path="$(EngineDir)/Binaries/Mac/UE4Game.app/Contents/MacOS/libogg.dylib" Type="DynamicLibrary" />
+    <BuildProduct Path="$(EngineDir)/Binaries/Mac/UE4Game.app/Contents/MacOS/libvorbis.dylib" Type="DynamicLibrary" />
+    <BuildProduct Path="$(EngineDir)/Plugins/Runtime/AppleMoviePlayer/Binaries/Mac/UE4Game-AppleMoviePlayer-Static.a" Type="StaticLibrary" />
+    <BuildProduct Path="$(EngineDir)/Binaries/Mac/UE4Game.app/Contents/MacOS/libogg.dylib" Type="DynamicLibrary" />
+    <BuildProduct Path="$(EngineDir)/Binaries/Mac/UE4Game.app/Contents/MacOS/libvorbis.dylib" Type="DynamicLibrary" />
+    <BuildProduct Path="$(EngineDir)/Plugins/Runtime/CableComponent/Binaries/Mac/UE4Game-CableComponent-Static.a" Type="StaticLibrary" />
+    <BuildProduct Path="$(EngineDir)/Binaries/Mac/UE4Game.app/Contents/MacOS/libogg.dylib" Type="DynamicLibrary" />
+    <BuildProduct Path="$(EngineDir)/Binaries/Mac/UE4Game.app/Contents/MacOS/libvorbis.dylib" Type="DynamicLibrary" />
+    <BuildProduct Path="$(EngineDir)/Plugins/Runtime/CustomMeshComponent/Binaries/Mac/UE4Game-CustomMeshComponent-Static.a" Type="StaticLibrary" />
+    <BuildProduct Path="$(EngineDir)/Binaries/Mac/UE4Game.app/Contents/MacOS/libogg.dylib" Type="DynamicLibrary" />
+    <BuildProduct Path="$(EngineDir)/Binaries/Mac/UE4Game.app/Contents/MacOS/libvorbis.dylib" Type="DynamicLibrary" />
+    <BuildProduct Path="$(EngineDir)/Plugins/Runtime/ExampleDeviceProfileSelector/Binaries/Mac/UE4Game-ExampleDeviceProfileSelector-Static.a" Type="StaticLibrary" />
+    <BuildProduct Path="$(EngineDir)/Binaries/Mac/UE4Game.app/Contents/MacOS/libogg.dylib" Type="DynamicLibrary" />
+    <BuildProduct Path="$(EngineDir)/Binaries/Mac/UE4Game.app/Contents/MacOS/libvorbis.dylib" Type="DynamicLibrary" />
+    <BuildProduct Path="$(EngineDir)/Plugins/Runtime/LiveEditorListenServer/Binaries/Mac/UE4Game-LiveEditorListenServer-Static.a" Type="StaticLibrary" />
+    <BuildProduct Path="$(EngineDir)/Binaries/Mac/UE4Game.app/Contents/MacOS/libogg.dylib" Type="DynamicLibrary" />
+    <BuildProduct Path="$(EngineDir)/Binaries/Mac/UE4Game.app/Contents/MacOS/libvorbis.dylib" Type="DynamicLibrary" />
+    <BuildProduct Path="$(EngineDir)/Plugins/Runtime/OculusRift/Binaries/Mac/UE4Game-OculusRift-Static.a" Type="StaticLibrary" />
+    <BuildProduct Path="$(EngineDir)/Binaries/Mac/UE4Game.app/Contents/MacOS/libogg.dylib" Type="DynamicLibrary" />
+    <BuildProduct Path="$(EngineDir)/Binaries/Mac/UE4Game.app/Contents/MacOS/libvorbis.dylib" Type="DynamicLibrary" />
+    <BuildProduct Path="$(EngineDir)/Plugins/Runtime/SoundMod/Binaries/Mac/UE4Game-SoundMod-Static.a" Type="StaticLibrary" />
+    <BuildProduct Path="$(EngineDir)/Binaries/Mac/UE4Game.app/Contents/MacOS/libogg.dylib" Type="DynamicLibrary" />
+    <BuildProduct Path="$(EngineDir)/Binaries/Mac/UE4Game.app/Contents/MacOS/libvorbis.dylib" Type="DynamicLibrary" />
+    <BuildProduct Path="$(EngineDir)/Plugins/Runtime/SoundVisualizations/Binaries/Mac/UE4Game-SoundVisualizations-Static.a" Type="StaticLibrary" />
+    <BuildProduct Path="$(EngineDir)/Binaries/Mac/UE4Game.app/Contents/MacOS/libogg.dylib" Type="DynamicLibrary" />
+    <BuildProduct Path="$(EngineDir)/Binaries/Mac/UE4Game.app/Contents/MacOS/libvorbis.dylib" Type="DynamicLibrary" />
+    <BuildProduct Path="$(EngineDir)/Plugins/Runtime/WebBrowserAsPlugin/Binaries/Mac/UE4Game-WebBrowserAsPlugin-Static.a" Type="StaticLibrary" />
+    <BuildProduct Path="$(EngineDir)/Binaries/Mac/UE4Game.app/Contents/MacOS/libplugin_carbon_interpose.dylib" Type="DynamicLibrary" />
+    <BuildProduct Path="$(EngineDir)/Binaries/Mac/UE4Game.app/Contents/MacOS/libogg.dylib" Type="DynamicLibrary" />
+    <BuildProduct Path="$(EngineDir)/Binaries/Mac/UE4Game.app/Contents/MacOS/libvorbis.dylib" Type="DynamicLibrary" />
+    <BuildProduct Path="$(EngineDir)/Binaries/Mac/UE4Game.app/Contents/Resources/am.lproj/locale.pak" Type="RequiredResource" />
+    <BuildProduct Path="$(EngineDir)/Binaries/Mac/UE4Game.app/Contents/Resources/ar.lproj/locale.pak" Type="RequiredResource" />
+    <BuildProduct Path="$(EngineDir)/Binaries/Mac/UE4Game.app/Contents/Resources/bg.lproj/locale.pak" Type="RequiredResource" />
+    <BuildProduct Path="$(EngineDir)/Binaries/Mac/UE4Game.app/Contents/Resources/bn.lproj/locale.pak" Type="RequiredResource" />
+    <BuildProduct Path="$(EngineDir)/Binaries/Mac/UE4Game.app/Contents/Resources/ca.lproj/locale.pak" Type="RequiredResource" />
+    <BuildProduct Path="$(EngineDir)/Binaries/Mac/UE4Game.app/Contents/Resources/cs.lproj/locale.pak" Type="RequiredResource" />
+    <BuildProduct Path="$(EngineDir)/Binaries/Mac/UE4Game.app/Contents/Resources/da.lproj/locale.pak" Type="RequiredResource" />
+    <BuildProduct Path="$(EngineDir)/Binaries/Mac/UE4Game.app/Contents/Resources/de.lproj/locale.pak" Type="RequiredResource" />
+    <BuildProduct Path="$(EngineDir)/Binaries/Mac/UE4Game.app/Contents/Resources/el.lproj/locale.pak" Type="RequiredResource" />
+    <BuildProduct Path="$(EngineDir)/Binaries/Mac/UE4Game.app/Contents/Resources/en.lproj/locale.pak" Type="RequiredResource" />
+    <BuildProduct Path="$(EngineDir)/Binaries/Mac/UE4Game.app/Contents/Resources/en_GB.lproj/locale.pak" Type="RequiredResource" />
+    <BuildProduct Path="$(EngineDir)/Binaries/Mac/UE4Game.app/Contents/Resources/es.lproj/locale.pak" Type="RequiredResource" />
+    <BuildProduct Path="$(EngineDir)/Binaries/Mac/UE4Game.app/Contents/Resources/es_419.lproj/locale.pak" Type="RequiredResource" />
+    <BuildProduct Path="$(EngineDir)/Binaries/Mac/UE4Game.app/Contents/Resources/et.lproj/locale.pak" Type="RequiredResource" />
+    <BuildProduct Path="$(EngineDir)/Binaries/Mac/UE4Game.app/Contents/Resources/fa.lproj/locale.pak" Type="RequiredResource" />
+    <BuildProduct Path="$(EngineDir)/Binaries/Mac/UE4Game.app/Contents/Resources/fi.lproj/locale.pak" Type="RequiredResource" />
+    <BuildProduct Path="$(EngineDir)/Binaries/Mac/UE4Game.app/Contents/Resources/fil.lproj/locale.pak" Type="RequiredResource" />
+    <BuildProduct Path="$(EngineDir)/Binaries/Mac/UE4Game.app/Contents/Resources/fr.lproj/locale.pak" Type="RequiredResource" />
+    <BuildProduct Path="$(EngineDir)/Binaries/Mac/UE4Game.app/Contents/Resources/gu.lproj/locale.pak" Type="RequiredResource" />
+    <BuildProduct Path="$(EngineDir)/Binaries/Mac/UE4Game.app/Contents/Resources/he.lproj/locale.pak" Type="RequiredResource" />
+    <BuildProduct Path="$(EngineDir)/Binaries/Mac/UE4Game.app/Contents/Resources/hi.lproj/locale.pak" Type="RequiredResource" />
+    <BuildProduct Path="$(EngineDir)/Binaries/Mac/UE4Game.app/Contents/Resources/hr.lproj/locale.pak" Type="RequiredResource" />
+    <BuildProduct Path="$(EngineDir)/Binaries/Mac/UE4Game.app/Contents/Resources/hu.lproj/locale.pak" Type="RequiredResource" />
+    <BuildProduct Path="$(EngineDir)/Binaries/Mac/UE4Game.app/Contents/Resources/id.lproj/locale.pak" Type="RequiredResource" />
+    <BuildProduct Path="$(EngineDir)/Binaries/Mac/UE4Game.app/Contents/Resources/it.lproj/locale.pak" Type="RequiredResource" />
+    <BuildProduct Path="$(EngineDir)/Binaries/Mac/UE4Game.app/Contents/Resources/ja.lproj/locale.pak" Type="RequiredResource" />
+    <BuildProduct Path="$(EngineDir)/Binaries/Mac/UE4Game.app/Contents/Resources/kn.lproj/locale.pak" Type="RequiredResource" />
+    <BuildProduct Path="$(EngineDir)/Binaries/Mac/UE4Game.app/Contents/Resources/ko.lproj/locale.pak" Type="RequiredResource" />
+    <BuildProduct Path="$(EngineDir)/Binaries/Mac/UE4Game.app/Contents/Resources/lt.lproj/locale.pak" Type="RequiredResource" />
+    <BuildProduct Path="$(EngineDir)/Binaries/Mac/UE4Game.app/Contents/Resources/lv.lproj/locale.pak" Type="RequiredResource" />
+    <BuildProduct Path="$(EngineDir)/Binaries/Mac/UE4Game.app/Contents/Resources/ml.lproj/locale.pak" Type="RequiredResource" />
+    <BuildProduct Path="$(EngineDir)/Binaries/Mac/UE4Game.app/Contents/Resources/mr.lproj/locale.pak" Type="RequiredResource" />
+    <BuildProduct Path="$(EngineDir)/Binaries/Mac/UE4Game.app/Contents/Resources/ms.lproj/locale.pak" Type="RequiredResource" />
+    <BuildProduct Path="$(EngineDir)/Binaries/Mac/UE4Game.app/Contents/Resources/nb.lproj/locale.pak" Type="RequiredResource" />
+    <BuildProduct Path="$(EngineDir)/Binaries/Mac/UE4Game.app/Contents/Resources/nl.lproj/locale.pak" Type="RequiredResource" />
+    <BuildProduct Path="$(EngineDir)/Binaries/Mac/UE4Game.app/Contents/Resources/pl.lproj/locale.pak" Type="RequiredResource" />
+    <BuildProduct Path="$(EngineDir)/Binaries/Mac/UE4Game.app/Contents/Resources/pt_BR.lproj/locale.pak" Type="RequiredResource" />
+    <BuildProduct Path="$(EngineDir)/Binaries/Mac/UE4Game.app/Contents/Resources/pt_PT.lproj/locale.pak" Type="RequiredResource" />
+    <BuildProduct Path="$(EngineDir)/Binaries/Mac/UE4Game.app/Contents/Resources/ro.lproj/locale.pak" Type="RequiredResource" />
+    <BuildProduct Path="$(EngineDir)/Binaries/Mac/UE4Game.app/Contents/Resources/ru.lproj/locale.pak" Type="RequiredResource" />
+    <BuildProduct Path="$(EngineDir)/Binaries/Mac/UE4Game.app/Contents/Resources/sk.lproj/locale.pak" Type="RequiredResource" />
+    <BuildProduct Path="$(EngineDir)/Binaries/Mac/UE4Game.app/Contents/Resources/sl.lproj/locale.pak" Type="RequiredResource" />
+    <BuildProduct Path="$(EngineDir)/Binaries/Mac/UE4Game.app/Contents/Resources/sr.lproj/locale.pak" Type="RequiredResource" />
+    <BuildProduct Path="$(EngineDir)/Binaries/Mac/UE4Game.app/Contents/Resources/sv.lproj/locale.pak" Type="RequiredResource" />
+    <BuildProduct Path="$(EngineDir)/Binaries/Mac/UE4Game.app/Contents/Resources/sw.lproj/locale.pak" Type="RequiredResource" />
+    <BuildProduct Path="$(EngineDir)/Binaries/Mac/UE4Game.app/Contents/Resources/ta.lproj/locale.pak" Type="RequiredResource" />
+    <BuildProduct Path="$(EngineDir)/Binaries/Mac/UE4Game.app/Contents/Resources/te.lproj/locale.pak" Type="RequiredResource" />
+    <BuildProduct Path="$(EngineDir)/Binaries/Mac/UE4Game.app/Contents/Resources/th.lproj/locale.pak" Type="RequiredResource" />
+    <BuildProduct Path="$(EngineDir)/Binaries/Mac/UE4Game.app/Contents/Resources/tr.lproj/locale.pak" Type="RequiredResource" />
+    <BuildProduct Path="$(EngineDir)/Binaries/Mac/UE4Game.app/Contents/Resources/uk.lproj/locale.pak" Type="RequiredResource" />
+    <BuildProduct Path="$(EngineDir)/Binaries/Mac/UE4Game.app/Contents/Resources/vi.lproj/locale.pak" Type="RequiredResource" />
+    <BuildProduct Path="$(EngineDir)/Binaries/Mac/UE4Game.app/Contents/Resources/zh_CN.lproj/locale.pak" Type="RequiredResource" />
+    <BuildProduct Path="$(EngineDir)/Binaries/Mac/UE4Game.app/Contents/Resources/zh_TW.lproj/locale.pak" Type="RequiredResource" />
+    <BuildProduct Path="$(EngineDir)/Plugins/ScriptPlugin/Binaries/Mac/UE4Game-ScriptPlugin-Static.a" Type="StaticLibrary" />
+    <BuildProduct Path="$(EngineDir)/Binaries/Mac/UE4Game.app/Contents/MacOS/libogg.dylib" Type="DynamicLibrary" />
+    <BuildProduct Path="$(EngineDir)/Binaries/Mac/UE4Game.app/Contents/MacOS/libvorbis.dylib" Type="DynamicLibrary" />
+    <BuildProduct Path="$(EngineDir)/Plugins/Slate/SlateRemote/Binaries/Mac/UE4Game-SlateRemote-Static.a" Type="StaticLibrary" />
+    <BuildProduct Path="$(EngineDir)/Binaries/Mac/UE4Game-Launch.a" Type="StaticLibrary" />
+    <BuildProduct Path="$(EngineDir)/Binaries/Mac/UE4Game.app/Contents/MacOS/libogg.dylib" Type="DynamicLibrary" />
+    <BuildProduct Path="$(EngineDir)/Binaries/Mac/UE4Game.app/Contents/MacOS/libvorbis.dylib" Type="DynamicLibrary" />
+    <BuildProduct Path="$(EngineDir)/Binaries/Mac/UE4Game-Core.a" Type="StaticLibrary" />
+    <BuildProduct Path="$(EngineDir)/Binaries/Mac/UE4Game-CoreUObject.a" Type="StaticLibrary" />
+    <BuildProduct Path="$(EngineDir)/Binaries/Mac/UE4Game-Projects.a" Type="StaticLibrary" />
+    <BuildProduct Path="$(EngineDir)/Binaries/Mac/UE4Game-Json.a" Type="StaticLibrary" />
+    <BuildProduct Path="$(EngineDir)/Binaries/Mac/UE4Game-Engine.a" Type="StaticLibrary" />
+    <BuildProduct Path="$(EngineDir)/Binaries/Mac/UE4Game.app/Contents/MacOS/libogg.dylib" Type="DynamicLibrary" />
+    <BuildProduct Path="$(EngineDir)/Binaries/Mac/UE4Game.app/Contents/MacOS/libvorbis.dylib" Type="DynamicLibrary" />
+    <BuildProduct Path="$(EngineDir)/Binaries/Mac/UE4Game-AppFramework.a" Type="StaticLibrary" />
+    <BuildProduct Path="$(EngineDir)/Binaries/Mac/UE4Game-InputCore.a" Type="StaticLibrary" />
+    <BuildProduct Path="$(EngineDir)/Binaries/Mac/UE4Game-SlateReflector.a" Type="StaticLibrary" />
+    <BuildProduct Path="$(EngineDir)/Binaries/Mac/UE4Game-Slate.a" Type="StaticLibrary" />
+    <BuildProduct Path="$(EngineDir)/Binaries/Mac/UE4Game-SlateCore.a" Type="StaticLibrary" />
+    <BuildProduct Path="$(EngineDir)/Binaries/Mac/UE4Game-Networking.a" Type="StaticLibrary" />
+    <BuildProduct Path="$(EngineDir)/Binaries/Mac/UE4Game-Sockets.a" Type="StaticLibrary" />
+    <BuildProduct Path="$(EngineDir)/Binaries/Mac/UE4Game-SandboxFile.a" Type="StaticLibrary" />
+    <BuildProduct Path="$(EngineDir)/Binaries/Mac/UE4Game-VectorVM.a" Type="StaticLibrary" />
+    <BuildProduct Path="$(EngineDir)/Binaries/Mac/UE4Game-Landscape.a" Type="StaticLibrary" />
+    <BuildProduct Path="$(EngineDir)/Binaries/Mac/UE4Game.app/Contents/MacOS/libogg.dylib" Type="DynamicLibrary" />
+    <BuildProduct Path="$(EngineDir)/Binaries/Mac/UE4Game.app/Contents/MacOS/libvorbis.dylib" Type="DynamicLibrary" />
+    <BuildProduct Path="$(EngineDir)/Binaries/Mac/UE4Game-RenderCore.a" Type="StaticLibrary" />
+    <BuildProduct Path="$(EngineDir)/Binaries/Mac/UE4Game-RHI.a" Type="StaticLibrary" />
+    <BuildProduct Path="$(EngineDir)/Binaries/Mac/UE4Game-NullDrv.a" Type="StaticLibrary" />
+    <BuildProduct Path="$(EngineDir)/Binaries/Mac/UE4Game-OpenGLDrv.a" Type="StaticLibrary" />
+    <BuildProduct Path="$(EngineDir)/Binaries/Mac/UE4Game.app/Contents/MacOS/libogg.dylib" Type="DynamicLibrary" />
+    <BuildProduct Path="$(EngineDir)/Binaries/Mac/UE4Game.app/Contents/MacOS/libvorbis.dylib" Type="DynamicLibrary" />
+    <BuildProduct Path="$(EngineDir)/Binaries/Mac/UE4Game-ShaderCore.a" Type="StaticLibrary" />
+    <BuildProduct Path="$(EngineDir)/Binaries/Mac/UE4Game-UtilityShaders.a" Type="StaticLibrary" />
+    <BuildProduct Path="$(EngineDir)/Binaries/Mac/UE4Game.app/Contents/MacOS/libogg.dylib" Type="DynamicLibrary" />
+    <BuildProduct Path="$(EngineDir)/Binaries/Mac/UE4Game.app/Contents/MacOS/libvorbis.dylib" Type="DynamicLibrary" />
+    <BuildProduct Path="$(EngineDir)/Binaries/Mac/UE4Game-ImageWrapper.a" Type="StaticLibrary" />
+    <BuildProduct Path="$(EngineDir)/Binaries/Mac/UE4Game-Renderer.a" Type="StaticLibrary" />
+    <BuildProduct Path="$(EngineDir)/Binaries/Mac/UE4Game.app/Contents/MacOS/libogg.dylib" Type="DynamicLibrary" />
+    <BuildProduct Path="$(EngineDir)/Binaries/Mac/UE4Game.app/Contents/MacOS/libvorbis.dylib" Type="DynamicLibrary" />
+    <BuildProduct Path="$(EngineDir)/Binaries/Mac/UE4Game-HeadMountedDisplay.a" Type="StaticLibrary" />
+    <BuildProduct Path="$(EngineDir)/Binaries/Mac/UE4Game.app/Contents/MacOS/libogg.dylib" Type="DynamicLibrary" />
+    <BuildProduct Path="$(EngineDir)/Binaries/Mac/UE4Game.app/Contents/MacOS/libvorbis.dylib" Type="DynamicLibrary" />
+    <BuildProduct Path="$(EngineDir)/Binaries/Mac/UE4Game-Foliage.a" Type="StaticLibrary" />
+    <BuildProduct Path="$(EngineDir)/Binaries/Mac/UE4Game.app/Contents/MacOS/libogg.dylib" Type="DynamicLibrary" />
+    <BuildProduct Path="$(EngineDir)/Binaries/Mac/UE4Game.app/Contents/MacOS/libvorbis.dylib" Type="DynamicLibrary" />
+    <BuildProduct Path="$(EngineDir)/Binaries/Mac/UE4Game-PhysXFormats.a" Type="StaticLibrary" />
+    <BuildProduct Path="$(EngineDir)/Binaries/Mac/UE4Game.app/Contents/MacOS/libogg.dylib" Type="DynamicLibrary" />
+    <BuildProduct Path="$(EngineDir)/Binaries/Mac/UE4Game.app/Contents/MacOS/libvorbis.dylib" Type="DynamicLibrary" />
+    <BuildProduct Path="$(EngineDir)/Binaries/Mac/UE4Game-UMG.a" Type="StaticLibrary" />
+    <BuildProduct Path="$(EngineDir)/Binaries/Mac/UE4Game.app/Contents/MacOS/libogg.dylib" Type="DynamicLibrary" />
+    <BuildProduct Path="$(EngineDir)/Binaries/Mac/UE4Game.app/Contents/MacOS/libvorbis.dylib" Type="DynamicLibrary" />
+    <BuildProduct Path="$(EngineDir)/Binaries/Mac/UE4Game-MovieSceneCore.a" Type="StaticLibrary" />
+    <BuildProduct Path="$(EngineDir)/Binaries/Mac/UE4Game.app/Contents/MacOS/libogg.dylib" Type="DynamicLibrary" />
+    <BuildProduct Path="$(EngineDir)/Binaries/Mac/UE4Game.app/Contents/MacOS/libvorbis.dylib" Type="DynamicLibrary" />
+    <BuildProduct Path="$(EngineDir)/Binaries/Mac/UE4Game-HTTP.a" Type="StaticLibrary" />
+    <BuildProduct Path="$(EngineDir)/Binaries/Mac/UE4Game-SlateRHIRenderer.a" Type="StaticLibrary" />
+    <BuildProduct Path="$(EngineDir)/Binaries/Mac/UE4Game.app/Contents/MacOS/libogg.dylib" Type="DynamicLibrary" />
+    <BuildProduct Path="$(EngineDir)/Binaries/Mac/UE4Game.app/Contents/MacOS/libvorbis.dylib" Type="DynamicLibrary" />
+    <BuildProduct Path="$(EngineDir)/Binaries/Mac/UE4Game-Navmesh.a" Type="StaticLibrary" />
+    <BuildProduct Path="$(EngineDir)/Binaries/Mac/UE4Game-Messaging.a" Type="StaticLibrary" />
+    <BuildProduct Path="$(EngineDir)/Binaries/Mac/UE4Game-AssetRegistry.a" Type="StaticLibrary" />
+    <BuildProduct Path="$(EngineDir)/Binaries/Mac/UE4Game-EngineMessages.a" Type="StaticLibrary" />
+    <BuildProduct Path="$(EngineDir)/Binaries/Mac/UE4Game-EngineSettings.a" Type="StaticLibrary" />
+    <BuildProduct Path="$(EngineDir)/Binaries/Mac/UE4Game-SynthBenchmark.a" Type="StaticLibrary" />
+    <BuildProduct Path="$(EngineDir)/Binaries/Mac/UE4Game-AIModule.a" Type="StaticLibrary" />
+    <BuildProduct Path="$(EngineDir)/Binaries/Mac/UE4Game.app/Contents/MacOS/libogg.dylib" Type="DynamicLibrary" />
+    <BuildProduct Path="$(EngineDir)/Binaries/Mac/UE4Game.app/Contents/MacOS/libvorbis.dylib" Type="DynamicLibrary" />
+    <BuildProduct Path="$(EngineDir)/Binaries/Mac/UE4Game-GameplayTags.a" Type="StaticLibrary" />
+    <BuildProduct Path="$(EngineDir)/Binaries/Mac/UE4Game.app/Contents/MacOS/libogg.dylib" Type="DynamicLibrary" />
+    <BuildProduct Path="$(EngineDir)/Binaries/Mac/UE4Game.app/Contents/MacOS/libvorbis.dylib" Type="DynamicLibrary" />
+    <BuildProduct Path="$(EngineDir)/Binaries/Mac/UE4Game-DatabaseSupport.a" Type="StaticLibrary" />
+    <BuildProduct Path="$(EngineDir)/Binaries/Mac/UE4Game-MovieSceneCoreTypes.a" Type="StaticLibrary" />
+    <BuildProduct Path="$(EngineDir)/Binaries/Mac/UE4Game.app/Contents/MacOS/libogg.dylib" Type="DynamicLibrary" />
+    <BuildProduct Path="$(EngineDir)/Binaries/Mac/UE4Game.app/Contents/MacOS/libvorbis.dylib" Type="DynamicLibrary" />
+    <BuildProduct Path="$(EngineDir)/Binaries/Mac/UE4Game-StreamingPauseRendering.a" Type="StaticLibrary" />
+    <BuildProduct Path="$(EngineDir)/Binaries/Mac/UE4Game.app/Contents/MacOS/libogg.dylib" Type="DynamicLibrary" />
+    <BuildProduct Path="$(EngineDir)/Binaries/Mac/UE4Game.app/Contents/MacOS/libvorbis.dylib" Type="DynamicLibrary" />
+    <BuildProduct Path="$(EngineDir)/Binaries/Mac/UE4Game-MoviePlayer.a" Type="StaticLibrary" />
+    <BuildProduct Path="$(EngineDir)/Binaries/Mac/UE4Game.app/Contents/MacOS/libogg.dylib" Type="DynamicLibrary" />
+    <BuildProduct Path="$(EngineDir)/Binaries/Mac/UE4Game.app/Contents/MacOS/libvorbis.dylib" Type="DynamicLibrary" />
+    <BuildProduct Path="$(EngineDir)/Binaries/Mac/UE4Game-Analytics.a" Type="StaticLibrary" />
+    <BuildProduct Path="$(EngineDir)/Binaries/Mac/UE4Game-AnalyticsET.a" Type="StaticLibrary" />
+    <BuildProduct Path="$(EngineDir)/Binaries/Mac/UE4Game-NetworkReplayStreaming.a" Type="StaticLibrary" />
+    <BuildProduct Path="$(EngineDir)/Binaries/Mac/UE4Game-NullNetworkReplayStreaming.a" Type="StaticLibrary" />
+    <BuildProduct Path="$(EngineDir)/Binaries/Mac/UE4Game-HttpNetworkReplayStreaming.a" Type="StaticLibrary" />
+    <BuildProduct Path="$(EngineDir)/Binaries/Mac/UE4Game.app/Contents/MacOS/libogg.dylib" Type="DynamicLibrary" />
+    <BuildProduct Path="$(EngineDir)/Binaries/Mac/UE4Game.app/Contents/MacOS/libvorbis.dylib" Type="DynamicLibrary" />
+    <BuildProduct Path="$(EngineDir)/Binaries/Mac/UE4Game-OnlineSubsystem.a" Type="StaticLibrary" />
+    <BuildProduct Path="$(EngineDir)/Binaries/Mac/UE4Game-ImageCore.a" Type="StaticLibrary" />
+    <BuildProduct Path="$(EngineDir)/Binaries/Mac/UE4Game-OnlineSubsystemUtils.a" Type="StaticLibrary" />
+    <BuildProduct Path="$(EngineDir)/Binaries/Mac/UE4Game.app/Contents/MacOS/libogg.dylib" Type="DynamicLibrary" />
+    <BuildProduct Path="$(EngineDir)/Binaries/Mac/UE4Game.app/Contents/MacOS/libvorbis.dylib" Type="DynamicLibrary" />
+    <BuildProduct Path="$(EngineDir)/Binaries/Mac/UE4Game-Voice.a" Type="StaticLibrary" />
+    <BuildProduct Path="$(EngineDir)/Binaries/Mac/UE4Game-Advertising.a" Type="StaticLibrary" />
+    <BuildProduct Path="$(EngineDir)/Binaries/Mac/UE4Game-CrashTracker.a" Type="StaticLibrary" />
+    <BuildProduct Path="$(EngineDir)/Binaries/Mac/UE4Game.app/Contents/MacOS/libogg.dylib" Type="DynamicLibrary" />
+    <BuildProduct Path="$(EngineDir)/Binaries/Mac/UE4Game.app/Contents/MacOS/libvorbis.dylib" Type="DynamicLibrary" />
+    <BuildProduct Path="$(EngineDir)/Binaries/Mac/UE4Game-GameLiveStreaming.a" Type="StaticLibrary" />
+    <BuildProduct Path="$(EngineDir)/Binaries/Mac/UE4Game.app/Contents/MacOS/libogg.dylib" Type="DynamicLibrary" />
+    <BuildProduct Path="$(EngineDir)/Binaries/Mac/UE4Game.app/Contents/MacOS/libvorbis.dylib" Type="DynamicLibrary" />
+    <BuildProduct Path="$(EngineDir)/Binaries/Mac/UE4Game-MediaAssets.a" Type="StaticLibrary" />
+    <BuildProduct Path="$(EngineDir)/Binaries/Mac/UE4Game.app/Contents/MacOS/libogg.dylib" Type="DynamicLibrary" />
+    <BuildProduct Path="$(EngineDir)/Binaries/Mac/UE4Game.app/Contents/MacOS/libvorbis.dylib" Type="DynamicLibrary" />
+    <BuildProduct Path="$(EngineDir)/Binaries/Mac/UE4Game-Media.a" Type="StaticLibrary" />
+    <BuildProduct Path="$(EngineDir)/Binaries/Mac/UE4Game-PakFile.a" Type="StaticLibrary" />
+    <BuildProduct Path="$(EngineDir)/Binaries/Mac/UE4Game-Serialization.a" Type="StaticLibrary" />
+    <BuildProduct Path="$(EngineDir)/Binaries/Mac/UE4Game-NetworkFile.a" Type="StaticLibrary" />
+    <BuildProduct Path="$(EngineDir)/Binaries/Mac/UE4Game-StreamingFile.a" Type="StaticLibrary" />
+    <BuildProduct Path="$(EngineDir)/Binaries/Mac/UE4Game-AutomationWorker.a" Type="StaticLibrary" />
+    <BuildProduct Path="$(EngineDir)/Binaries/Mac/UE4Game.app/Contents/MacOS/libogg.dylib" Type="DynamicLibrary" />
+    <BuildProduct Path="$(EngineDir)/Binaries/Mac/UE4Game.app/Contents/MacOS/libvorbis.dylib" Type="DynamicLibrary" />
+    <BuildProduct Path="$(EngineDir)/Binaries/Mac/UE4Game-AutomationMessages.a" Type="StaticLibrary" />
+    <BuildProduct Path="$(EngineDir)/Binaries/Mac/UE4Game-SessionServices.a" Type="StaticLibrary" />
+    <BuildProduct Path="$(EngineDir)/Binaries/Mac/UE4Game-SessionMessages.a" Type="StaticLibrary" />
+    <BuildProduct Path="$(EngineDir)/Binaries/Mac/UE4Game-CoreAudio.a" Type="StaticLibrary" />
+    <BuildProduct Path="$(EngineDir)/Binaries/Mac/UE4Game.app/Contents/MacOS/libogg.dylib" Type="DynamicLibrary" />
+    <BuildProduct Path="$(EngineDir)/Binaries/Mac/UE4Game.app/Contents/MacOS/libvorbis.dylib" Type="DynamicLibrary" />
+    <BuildProduct Path="$(EngineDir)/Binaries/Mac/UE4Game.app/Contents/Resources/RadioEffectUnit.component/Contents/Info.plist" Type="RequiredResource" />
+    <BuildProduct Path="$(EngineDir)/Binaries/Mac/UE4Game.app/Contents/Resources/RadioEffectUnit.component/Contents/MacOS/RadioEffectUnit" Type="RequiredResource" />
+    <BuildProduct Path="$(EngineDir)/Binaries/Mac/UE4Game.app/Contents/Resources/RadioEffectUnit.component/Contents/Resources/English.lproj/Localizable.strings" Type="RequiredResource" />
+    <BuildProduct Path="$(EngineDir)/Binaries/Mac/UE4Game-TaskGraph.a" Type="StaticLibrary" />
+    <BuildProduct Path="$(EngineDir)/Binaries/Mac/UE4Game.app/Contents/MacOS/libogg.dylib" Type="DynamicLibrary" />
+    <BuildProduct Path="$(EngineDir)/Binaries/Mac/UE4Game.app/Contents/MacOS/libvorbis.dylib" Type="DynamicLibrary" />
+    <BuildProduct Path="$(EngineDir)/Binaries/Mac/UE4Game-EditorStyle.a" Type="StaticLibrary" />
+    <BuildProduct Path="$(EngineDir)/Binaries/Mac/UE4Game-RealtimeProfiler.a" Type="StaticLibrary" />
+    <BuildProduct Path="$(EngineDir)/Binaries/Mac/UE4Game.app/Contents/MacOS/libogg.dylib" Type="DynamicLibrary" />
+    <BuildProduct Path="$(EngineDir)/Binaries/Mac/UE4Game.app/Contents/MacOS/libvorbis.dylib" Type="DynamicLibrary" />
+    <BuildProduct Path="$(EngineDir)/Binaries/Mac/UE4Game-ProfilerService.a" Type="StaticLibrary" />
+    <BuildProduct Path="$(EngineDir)/Binaries/Mac/UE4Game.app/Contents/MacOS/libogg.dylib" Type="DynamicLibrary" />
+    <BuildProduct Path="$(EngineDir)/Binaries/Mac/UE4Game.app/Contents/MacOS/libvorbis.dylib" Type="DynamicLibrary" />
+    <BuildProduct Path="$(EngineDir)/Binaries/Mac/UE4Game-ProfilerMessages.a" Type="StaticLibrary" />
+    <BuildProduct Path="$(EngineDir)/Binaries/Mac/UE4Game-InputDevice.a" Type="StaticLibrary" />
+    <BuildProduct Path="$(EngineDir)/Binaries/Mac/UE4Game.app/Contents/MacOS/libogg.dylib" Type="DynamicLibrary" />
+    <BuildProduct Path="$(EngineDir)/Binaries/Mac/UE4Game.app/Contents/MacOS/libvorbis.dylib" Type="DynamicLibrary" />
+    <BuildProduct Path="$(EngineDir)/Binaries/Mac/UE4Game-GameMenuBuilder.a" Type="StaticLibrary" />
+    <BuildProduct Path="$(EngineDir)/Binaries/Mac/UE4Game.app/Contents/MacOS/libogg.dylib" Type="DynamicLibrary" />
+    <BuildProduct Path="$(EngineDir)/Binaries/Mac/UE4Game.app/Contents/MacOS/libvorbis.dylib" Type="DynamicLibrary" />
+    <BuildProduct Path="$(EngineDir)/Binaries/Mac/UE4Game-XmlParser.a" Type="StaticLibrary" />
+    <BuildProduct Path="$(EngineDir)/Binaries/Mac/UE4Game-OnlineSubsystemNull.a" Type="StaticLibrary" />
+    <BuildProduct Path="$(EngineDir)/Binaries/Mac/UE4Game.app/Contents/MacOS/libogg.dylib" Type="DynamicLibrary" />
+    <BuildProduct Path="$(EngineDir)/Binaries/Mac/UE4Game.app/Contents/MacOS/libvorbis.dylib" Type="DynamicLibrary" />
+    <BuildProduct Path="$(EngineDir)/Binaries/Mac/UE4Game-OnlineSubsystemSteam.a" Type="StaticLibrary" />
+    <BuildProduct Path="$(EngineDir)/Binaries/Mac/UE4Game.app/Contents/MacOS/libogg.dylib" Type="DynamicLibrary" />
+    <BuildProduct Path="$(EngineDir)/Binaries/Mac/UE4Game.app/Contents/MacOS/libvorbis.dylib" Type="DynamicLibrary" />
+    <BuildProduct Path="$(EngineDir)/Binaries/Mac/UE4Game-WebBrowser.a" Type="StaticLibrary" />
+    <BuildProduct Path="$(EngineDir)/Binaries/Mac/UE4Game.app/Contents/MacOS/libplugin_carbon_interpose.dylib" Type="DynamicLibrary" />
+    <BuildProduct Path="$(EngineDir)/Binaries/Mac/UE4Game.app/Contents/Resources/am.lproj/locale.pak" Type="RequiredResource" />
+    <BuildProduct Path="$(EngineDir)/Binaries/Mac/UE4Game.app/Contents/Resources/ar.lproj/locale.pak" Type="RequiredResource" />
+    <BuildProduct Path="$(EngineDir)/Binaries/Mac/UE4Game.app/Contents/Resources/bg.lproj/locale.pak" Type="RequiredResource" />
+    <BuildProduct Path="$(EngineDir)/Binaries/Mac/UE4Game.app/Contents/Resources/bn.lproj/locale.pak" Type="RequiredResource" />
+    <BuildProduct Path="$(EngineDir)/Binaries/Mac/UE4Game.app/Contents/Resources/ca.lproj/locale.pak" Type="RequiredResource" />
+    <BuildProduct Path="$(EngineDir)/Binaries/Mac/UE4Game.app/Contents/Resources/cs.lproj/locale.pak" Type="RequiredResource" />
+    <BuildProduct Path="$(EngineDir)/Binaries/Mac/UE4Game.app/Contents/Resources/da.lproj/locale.pak" Type="RequiredResource" />
+    <BuildProduct Path="$(EngineDir)/Binaries/Mac/UE4Game.app/Contents/Resources/de.lproj/locale.pak" Type="RequiredResource" />
+    <BuildProduct Path="$(EngineDir)/Binaries/Mac/UE4Game.app/Contents/Resources/el.lproj/locale.pak" Type="RequiredResource" />
+    <BuildProduct Path="$(EngineDir)/Binaries/Mac/UE4Game.app/Contents/Resources/en.lproj/locale.pak" Type="RequiredResource" />
+    <BuildProduct Path="$(EngineDir)/Binaries/Mac/UE4Game.app/Contents/Resources/en_GB.lproj/locale.pak" Type="RequiredResource" />
+    <BuildProduct Path="$(EngineDir)/Binaries/Mac/UE4Game.app/Contents/Resources/es.lproj/locale.pak" Type="RequiredResource" />
+    <BuildProduct Path="$(EngineDir)/Binaries/Mac/UE4Game.app/Contents/Resources/es_419.lproj/locale.pak" Type="RequiredResource" />
+    <BuildProduct Path="$(EngineDir)/Binaries/Mac/UE4Game.app/Contents/Resources/et.lproj/locale.pak" Type="RequiredResource" />
+    <BuildProduct Path="$(EngineDir)/Binaries/Mac/UE4Game.app/Contents/Resources/fa.lproj/locale.pak" Type="RequiredResource" />
+    <BuildProduct Path="$(EngineDir)/Binaries/Mac/UE4Game.app/Contents/Resources/fi.lproj/locale.pak" Type="RequiredResource" />
+    <BuildProduct Path="$(EngineDir)/Binaries/Mac/UE4Game.app/Contents/Resources/fil.lproj/locale.pak" Type="RequiredResource" />
+    <BuildProduct Path="$(EngineDir)/Binaries/Mac/UE4Game.app/Contents/Resources/fr.lproj/locale.pak" Type="RequiredResource" />
+    <BuildProduct Path="$(EngineDir)/Binaries/Mac/UE4Game.app/Contents/Resources/gu.lproj/locale.pak" Type="RequiredResource" />
+    <BuildProduct Path="$(EngineDir)/Binaries/Mac/UE4Game.app/Contents/Resources/he.lproj/locale.pak" Type="RequiredResource" />
+    <BuildProduct Path="$(EngineDir)/Binaries/Mac/UE4Game.app/Contents/Resources/hi.lproj/locale.pak" Type="RequiredResource" />
+    <BuildProduct Path="$(EngineDir)/Binaries/Mac/UE4Game.app/Contents/Resources/hr.lproj/locale.pak" Type="RequiredResource" />
+    <BuildProduct Path="$(EngineDir)/Binaries/Mac/UE4Game.app/Contents/Resources/hu.lproj/locale.pak" Type="RequiredResource" />
+    <BuildProduct Path="$(EngineDir)/Binaries/Mac/UE4Game.app/Contents/Resources/id.lproj/locale.pak" Type="RequiredResource" />
+    <BuildProduct Path="$(EngineDir)/Binaries/Mac/UE4Game.app/Contents/Resources/it.lproj/locale.pak" Type="RequiredResource" />
+    <BuildProduct Path="$(EngineDir)/Binaries/Mac/UE4Game.app/Contents/Resources/ja.lproj/locale.pak" Type="RequiredResource" />
+    <BuildProduct Path="$(EngineDir)/Binaries/Mac/UE4Game.app/Contents/Resources/kn.lproj/locale.pak" Type="RequiredResource" />
+    <BuildProduct Path="$(EngineDir)/Binaries/Mac/UE4Game.app/Contents/Resources/ko.lproj/locale.pak" Type="RequiredResource" />
+    <BuildProduct Path="$(EngineDir)/Binaries/Mac/UE4Game.app/Contents/Resources/lt.lproj/locale.pak" Type="RequiredResource" />
+    <BuildProduct Path="$(EngineDir)/Binaries/Mac/UE4Game.app/Contents/Resources/lv.lproj/locale.pak" Type="RequiredResource" />
+    <BuildProduct Path="$(EngineDir)/Binaries/Mac/UE4Game.app/Contents/Resources/ml.lproj/locale.pak" Type="RequiredResource" />
+    <BuildProduct Path="$(EngineDir)/Binaries/Mac/UE4Game.app/Contents/Resources/mr.lproj/locale.pak" Type="RequiredResource" />
+    <BuildProduct Path="$(EngineDir)/Binaries/Mac/UE4Game.app/Contents/Resources/ms.lproj/locale.pak" Type="RequiredResource" />
+    <BuildProduct Path="$(EngineDir)/Binaries/Mac/UE4Game.app/Contents/Resources/nb.lproj/locale.pak" Type="RequiredResource" />
+    <BuildProduct Path="$(EngineDir)/Binaries/Mac/UE4Game.app/Contents/Resources/nl.lproj/locale.pak" Type="RequiredResource" />
+    <BuildProduct Path="$(EngineDir)/Binaries/Mac/UE4Game.app/Contents/Resources/pl.lproj/locale.pak" Type="RequiredResource" />
+    <BuildProduct Path="$(EngineDir)/Binaries/Mac/UE4Game.app/Contents/Resources/pt_BR.lproj/locale.pak" Type="RequiredResource" />
+    <BuildProduct Path="$(EngineDir)/Binaries/Mac/UE4Game.app/Contents/Resources/pt_PT.lproj/locale.pak" Type="RequiredResource" />
+    <BuildProduct Path="$(EngineDir)/Binaries/Mac/UE4Game.app/Contents/Resources/ro.lproj/locale.pak" Type="RequiredResource" />
+    <BuildProduct Path="$(EngineDir)/Binaries/Mac/UE4Game.app/Contents/Resources/ru.lproj/locale.pak" Type="RequiredResource" />
+    <BuildProduct Path="$(EngineDir)/Binaries/Mac/UE4Game.app/Contents/Resources/sk.lproj/locale.pak" Type="RequiredResource" />
+    <BuildProduct Path="$(EngineDir)/Binaries/Mac/UE4Game.app/Contents/Resources/sl.lproj/locale.pak" Type="RequiredResource" />
+    <BuildProduct Path="$(EngineDir)/Binaries/Mac/UE4Game.app/Contents/Resources/sr.lproj/locale.pak" Type="RequiredResource" />
+    <BuildProduct Path="$(EngineDir)/Binaries/Mac/UE4Game.app/Contents/Resources/sv.lproj/locale.pak" Type="RequiredResource" />
+    <BuildProduct Path="$(EngineDir)/Binaries/Mac/UE4Game.app/Contents/Resources/sw.lproj/locale.pak" Type="RequiredResource" />
+    <BuildProduct Path="$(EngineDir)/Binaries/Mac/UE4Game.app/Contents/Resources/ta.lproj/locale.pak" Type="RequiredResource" />
+    <BuildProduct Path="$(EngineDir)/Binaries/Mac/UE4Game.app/Contents/Resources/te.lproj/locale.pak" Type="RequiredResource" />
+    <BuildProduct Path="$(EngineDir)/Binaries/Mac/UE4Game.app/Contents/Resources/th.lproj/locale.pak" Type="RequiredResource" />
+    <BuildProduct Path="$(EngineDir)/Binaries/Mac/UE4Game.app/Contents/Resources/tr.lproj/locale.pak" Type="RequiredResource" />
+    <BuildProduct Path="$(EngineDir)/Binaries/Mac/UE4Game.app/Contents/Resources/uk.lproj/locale.pak" Type="RequiredResource" />
+    <BuildProduct Path="$(EngineDir)/Binaries/Mac/UE4Game.app/Contents/Resources/vi.lproj/locale.pak" Type="RequiredResource" />
+    <BuildProduct Path="$(EngineDir)/Binaries/Mac/UE4Game.app/Contents/Resources/zh_CN.lproj/locale.pak" Type="RequiredResource" />
+    <BuildProduct Path="$(EngineDir)/Binaries/Mac/UE4Game.app/Contents/Resources/zh_TW.lproj/locale.pak" Type="RequiredResource" />
+    <BuildProduct Path="$(EngineDir)/Binaries/Mac/UE4Game-CEF3Utils.a" Type="StaticLibrary" />
+    <BuildProduct Path="$(EngineDir)/Binaries/Mac/UE4Game.app/Contents/MacOS/libplugin_carbon_interpose.dylib" Type="DynamicLibrary" />
+    <BuildProduct Path="$(EngineDir)/Binaries/Mac/UE4Game.app/Contents/Resources/am.lproj/locale.pak" Type="RequiredResource" />
+    <BuildProduct Path="$(EngineDir)/Binaries/Mac/UE4Game.app/Contents/Resources/ar.lproj/locale.pak" Type="RequiredResource" />
+    <BuildProduct Path="$(EngineDir)/Binaries/Mac/UE4Game.app/Contents/Resources/bg.lproj/locale.pak" Type="RequiredResource" />
+    <BuildProduct Path="$(EngineDir)/Binaries/Mac/UE4Game.app/Contents/Resources/bn.lproj/locale.pak" Type="RequiredResource" />
+    <BuildProduct Path="$(EngineDir)/Binaries/Mac/UE4Game.app/Contents/Resources/ca.lproj/locale.pak" Type="RequiredResource" />
+    <BuildProduct Path="$(EngineDir)/Binaries/Mac/UE4Game.app/Contents/Resources/cs.lproj/locale.pak" Type="RequiredResource" />
+    <BuildProduct Path="$(EngineDir)/Binaries/Mac/UE4Game.app/Contents/Resources/da.lproj/locale.pak" Type="RequiredResource" />
+    <BuildProduct Path="$(EngineDir)/Binaries/Mac/UE4Game.app/Contents/Resources/de.lproj/locale.pak" Type="RequiredResource" />
+    <BuildProduct Path="$(EngineDir)/Binaries/Mac/UE4Game.app/Contents/Resources/el.lproj/locale.pak" Type="RequiredResource" />
+    <BuildProduct Path="$(EngineDir)/Binaries/Mac/UE4Game.app/Contents/Resources/en.lproj/locale.pak" Type="RequiredResource" />
+    <BuildProduct Path="$(EngineDir)/Binaries/Mac/UE4Game.app/Contents/Resources/en_GB.lproj/locale.pak" Type="RequiredResource" />
+    <BuildProduct Path="$(EngineDir)/Binaries/Mac/UE4Game.app/Contents/Resources/es.lproj/locale.pak" Type="RequiredResource" />
+    <BuildProduct Path="$(EngineDir)/Binaries/Mac/UE4Game.app/Contents/Resources/es_419.lproj/locale.pak" Type="RequiredResource" />
+    <BuildProduct Path="$(EngineDir)/Binaries/Mac/UE4Game.app/Contents/Resources/et.lproj/locale.pak" Type="RequiredResource" />
+    <BuildProduct Path="$(EngineDir)/Binaries/Mac/UE4Game.app/Contents/Resources/fa.lproj/locale.pak" Type="RequiredResource" />
+    <BuildProduct Path="$(EngineDir)/Binaries/Mac/UE4Game.app/Contents/Resources/fi.lproj/locale.pak" Type="RequiredResource" />
+    <BuildProduct Path="$(EngineDir)/Binaries/Mac/UE4Game.app/Contents/Resources/fil.lproj/locale.pak" Type="RequiredResource" />
+    <BuildProduct Path="$(EngineDir)/Binaries/Mac/UE4Game.app/Contents/Resources/fr.lproj/locale.pak" Type="RequiredResource" />
+    <BuildProduct Path="$(EngineDir)/Binaries/Mac/UE4Game.app/Contents/Resources/gu.lproj/locale.pak" Type="RequiredResource" />
+    <BuildProduct Path="$(EngineDir)/Binaries/Mac/UE4Game.app/Contents/Resources/he.lproj/locale.pak" Type="RequiredResource" />
+    <BuildProduct Path="$(EngineDir)/Binaries/Mac/UE4Game.app/Contents/Resources/hi.lproj/locale.pak" Type="RequiredResource" />
+    <BuildProduct Path="$(EngineDir)/Binaries/Mac/UE4Game.app/Contents/Resources/hr.lproj/locale.pak" Type="RequiredResource" />
+    <BuildProduct Path="$(EngineDir)/Binaries/Mac/UE4Game.app/Contents/Resources/hu.lproj/locale.pak" Type="RequiredResource" />
+    <BuildProduct Path="$(EngineDir)/Binaries/Mac/UE4Game.app/Contents/Resources/id.lproj/locale.pak" Type="RequiredResource" />
+    <BuildProduct Path="$(EngineDir)/Binaries/Mac/UE4Game.app/Contents/Resources/it.lproj/locale.pak" Type="RequiredResource" />
+    <BuildProduct Path="$(EngineDir)/Binaries/Mac/UE4Game.app/Contents/Resources/ja.lproj/locale.pak" Type="RequiredResource" />
+    <BuildProduct Path="$(EngineDir)/Binaries/Mac/UE4Game.app/Contents/Resources/kn.lproj/locale.pak" Type="RequiredResource" />
+    <BuildProduct Path="$(EngineDir)/Binaries/Mac/UE4Game.app/Contents/Resources/ko.lproj/locale.pak" Type="RequiredResource" />
+    <BuildProduct Path="$(EngineDir)/Binaries/Mac/UE4Game.app/Contents/Resources/lt.lproj/locale.pak" Type="RequiredResource" />
+    <BuildProduct Path="$(EngineDir)/Binaries/Mac/UE4Game.app/Contents/Resources/lv.lproj/locale.pak" Type="RequiredResource" />
+    <BuildProduct Path="$(EngineDir)/Binaries/Mac/UE4Game.app/Contents/Resources/ml.lproj/locale.pak" Type="RequiredResource" />
+    <BuildProduct Path="$(EngineDir)/Binaries/Mac/UE4Game.app/Contents/Resources/mr.lproj/locale.pak" Type="RequiredResource" />
+    <BuildProduct Path="$(EngineDir)/Binaries/Mac/UE4Game.app/Contents/Resources/ms.lproj/locale.pak" Type="RequiredResource" />
+    <BuildProduct Path="$(EngineDir)/Binaries/Mac/UE4Game.app/Contents/Resources/nb.lproj/locale.pak" Type="RequiredResource" />
+    <BuildProduct Path="$(EngineDir)/Binaries/Mac/UE4Game.app/Contents/Resources/nl.lproj/locale.pak" Type="RequiredResource" />
+    <BuildProduct Path="$(EngineDir)/Binaries/Mac/UE4Game.app/Contents/Resources/pl.lproj/locale.pak" Type="RequiredResource" />
+    <BuildProduct Path="$(EngineDir)/Binaries/Mac/UE4Game.app/Contents/Resources/pt_BR.lproj/locale.pak" Type="RequiredResource" />
+    <BuildProduct Path="$(EngineDir)/Binaries/Mac/UE4Game.app/Contents/Resources/pt_PT.lproj/locale.pak" Type="RequiredResource" />
+    <BuildProduct Path="$(EngineDir)/Binaries/Mac/UE4Game.app/Contents/Resources/ro.lproj/locale.pak" Type="RequiredResource" />
+    <BuildProduct Path="$(EngineDir)/Binaries/Mac/UE4Game.app/Contents/Resources/ru.lproj/locale.pak" Type="RequiredResource" />
+    <BuildProduct Path="$(EngineDir)/Binaries/Mac/UE4Game.app/Contents/Resources/sk.lproj/locale.pak" Type="RequiredResource" />
+    <BuildProduct Path="$(EngineDir)/Binaries/Mac/UE4Game.app/Contents/Resources/sl.lproj/locale.pak" Type="RequiredResource" />
+    <BuildProduct Path="$(EngineDir)/Binaries/Mac/UE4Game.app/Contents/Resources/sr.lproj/locale.pak" Type="RequiredResource" />
+    <BuildProduct Path="$(EngineDir)/Binaries/Mac/UE4Game.app/Contents/Resources/sv.lproj/locale.pak" Type="RequiredResource" />
+    <BuildProduct Path="$(EngineDir)/Binaries/Mac/UE4Game.app/Contents/Resources/sw.lproj/locale.pak" Type="RequiredResource" />
+    <BuildProduct Path="$(EngineDir)/Binaries/Mac/UE4Game.app/Contents/Resources/ta.lproj/locale.pak" Type="RequiredResource" />
+    <BuildProduct Path="$(EngineDir)/Binaries/Mac/UE4Game.app/Contents/Resources/te.lproj/locale.pak" Type="RequiredResource" />
+    <BuildProduct Path="$(EngineDir)/Binaries/Mac/UE4Game.app/Contents/Resources/th.lproj/locale.pak" Type="RequiredResource" />
+    <BuildProduct Path="$(EngineDir)/Binaries/Mac/UE4Game.app/Contents/Resources/tr.lproj/locale.pak" Type="RequiredResource" />
+    <BuildProduct Path="$(EngineDir)/Binaries/Mac/UE4Game.app/Contents/Resources/uk.lproj/locale.pak" Type="RequiredResource" />
+    <BuildProduct Path="$(EngineDir)/Binaries/Mac/UE4Game.app/Contents/Resources/vi.lproj/locale.pak" Type="RequiredResource" />
+    <BuildProduct Path="$(EngineDir)/Binaries/Mac/UE4Game.app/Contents/Resources/zh_CN.lproj/locale.pak" Type="RequiredResource" />
+    <BuildProduct Path="$(EngineDir)/Binaries/Mac/UE4Game.app/Contents/Resources/zh_TW.lproj/locale.pak" Type="RequiredResource" />
+  </BuildProducts>
+  <RuntimeDependencies />
+</BuildReceipt>
\ No newline at end of file
