GUBP Node Shared_LabelPromotable built from changelist 2506329

diff --git a/Engine/Build/Receipts/UE4Client-Win64-Development.target.xml b/Engine/Build/Receipts/UE4Client-Win64-Development.target.xml
--- a/Engine/Build/Receipts/UE4Client-Win64-Development.target.xml
+++ b/Engine/Build/Receipts/UE4Client-Win64-Development.target.xml
@@ -11,6 +11,7 @@
     <BuildProduct Path="$(EngineDir)\Plugins\Runtime\CustomMeshComponent\Binaries\Win64\UE4Client-CustomMeshComponent-Static.lib" Type="StaticLibrary" />
     <BuildProduct Path="$(EngineDir)\Plugins\Runtime\ExampleDeviceProfileSelector\Binaries\Win64\UE4Client-ExampleDeviceProfileSelector-Static.lib" Type="StaticLibrary" />
     <BuildProduct Path="$(EngineDir)\Plugins\Runtime\OculusRift\Binaries\Win64\UE4Client-OculusRift-Static.lib" Type="StaticLibrary" />
+    <BuildProduct Path="$(EngineDir)\Plugins\Runtime\PS4\Morpheus\Binaries\Win64\UE4Client-Morpheus-Static.lib" Type="StaticLibrary" />
     <BuildProduct Path="$(EngineDir)\Plugins\Runtime\WindowsMoviePlayer\Binaries\Win64\UE4Client-WindowsMoviePlayer-Static.lib" Type="StaticLibrary" />
     <BuildProduct Path="$(EngineDir)\Plugins\Slate\SlateRemote\Binaries\Win64\UE4Client-SlateRemote-Static.lib" Type="StaticLibrary" />
   </BuildProducts>
diff --git a/Engine/Build/Receipts/UE4Game-HTML5-Development.target.xml b/Engine/Build/Receipts/UE4Game-HTML5-Development.target.xml
--- a/Engine/Build/Receipts/UE4Game-HTML5-Development.target.xml
+++ b/Engine/Build/Receipts/UE4Game-HTML5-Development.target.xml
@@ -16,6 +16,7 @@
     <BuildProduct Path="$(EngineDir)\Binaries\HTML5\UE4Game-Launch.bc" Type="StaticLibrary" />
     <BuildProduct Path="$(EngineDir)\Binaries\HTML5\UE4Game-Core.bc" Type="StaticLibrary" />
     <BuildProduct Path="$(EngineDir)\Binaries\HTML5\UE4Game-HTML5JS.bc" Type="StaticLibrary" />
+    <BuildProduct Path="$(EngineDir)\Binaries\HTML5\UE4Game-MapPakDownloader.bc" Type="StaticLibrary" />
     <BuildProduct Path="$(EngineDir)\Binaries\HTML5\UE4Game-CoreUObject.bc" Type="StaticLibrary" />
     <BuildProduct Path="$(EngineDir)\Binaries\HTML5\UE4Game-Projects.bc" Type="StaticLibrary" />
     <BuildProduct Path="$(EngineDir)\Binaries\HTML5\UE4Game-Json.bc" Type="StaticLibrary" />
diff --git a/Engine/Build/Receipts/UE4Game-HTML5-Shipping.target.xml b/Engine/Build/Receipts/UE4Game-HTML5-Shipping.target.xml
--- a/Engine/Build/Receipts/UE4Game-HTML5-Shipping.target.xml
+++ b/Engine/Build/Receipts/UE4Game-HTML5-Shipping.target.xml
@@ -16,6 +16,7 @@
     <BuildProduct Path="$(EngineDir)\Binaries\HTML5\UE4Game-Launch-HTML5-Shipping.bc" Type="StaticLibrary" />
     <BuildProduct Path="$(EngineDir)\Binaries\HTML5\UE4Game-Core-HTML5-Shipping.bc" Type="StaticLibrary" />
     <BuildProduct Path="$(EngineDir)\Binaries\HTML5\UE4Game-HTML5JS-HTML5-Shipping.bc" Type="StaticLibrary" />
+    <BuildProduct Path="$(EngineDir)\Binaries\HTML5\UE4Game-MapPakDownloader-HTML5-Shipping.bc" Type="StaticLibrary" />
     <BuildProduct Path="$(EngineDir)\Binaries\HTML5\UE4Game-CoreUObject-HTML5-Shipping.bc" Type="StaticLibrary" />
     <BuildProduct Path="$(EngineDir)\Binaries\HTML5\UE4Game-Projects-HTML5-Shipping.bc" Type="StaticLibrary" />
     <BuildProduct Path="$(EngineDir)\Binaries\HTML5\UE4Game-Json-HTML5-Shipping.bc" Type="StaticLibrary" />
diff --git a/Engine/Build/Receipts/UE4Game-PS4-Development.target.xml b/Engine/Build/Receipts/UE4Game-PS4-Development.target.xml
--- a/Engine/Build/Receipts/UE4Game-PS4-Development.target.xml
+++ b/Engine/Build/Receipts/UE4Game-PS4-Development.target.xml
@@ -9,6 +9,7 @@
     <BuildProduct Path="$(EngineDir)\Plugins\Runtime\CableComponent\Binaries\PS4\UE4Game-CableComponent-Static.a" Type="StaticLibrary" />
     <BuildProduct Path="$(EngineDir)\Plugins\Runtime\CustomMeshComponent\Binaries\PS4\UE4Game-CustomMeshComponent-Static.a" Type="StaticLibrary" />
     <BuildProduct Path="$(EngineDir)\Plugins\Runtime\ExampleDeviceProfileSelector\Binaries\PS4\UE4Game-ExampleDeviceProfileSelector-Static.a" Type="StaticLibrary" />
+    <BuildProduct Path="$(EngineDir)\Plugins\Runtime\PS4\Morpheus\Binaries\PS4\UE4Game-Morpheus-Static.a" Type="StaticLibrary" />
     <BuildProduct Path="$(EngineDir)\Plugins\Runtime\PS4\PS4MoviePlayer\Binaries\PS4\UE4Game-PS4MoviePlayer-Static.a" Type="StaticLibrary" />
     <BuildProduct Path="$(EngineDir)\Plugins\Slate\SlateRemote\Binaries\PS4\UE4Game-SlateRemote-Static.a" Type="StaticLibrary" />
   </BuildProducts>
diff --git a/Engine/Build/Receipts/UE4Game-PS4-Shipping.target.xml b/Engine/Build/Receipts/UE4Game-PS4-Shipping.target.xml
--- a/Engine/Build/Receipts/UE4Game-PS4-Shipping.target.xml
+++ b/Engine/Build/Receipts/UE4Game-PS4-Shipping.target.xml
@@ -9,6 +9,7 @@
     <BuildProduct Path="$(EngineDir)\Plugins\Runtime\CableComponent\Binaries\PS4\UE4Game-CableComponent-PS4-Shipping-Static.a" Type="StaticLibrary" />
     <BuildProduct Path="$(EngineDir)\Plugins\Runtime\CustomMeshComponent\Binaries\PS4\UE4Game-CustomMeshComponent-PS4-Shipping-Static.a" Type="StaticLibrary" />
     <BuildProduct Path="$(EngineDir)\Plugins\Runtime\ExampleDeviceProfileSelector\Binaries\PS4\UE4Game-ExampleDeviceProfileSelector-PS4-Shipping-Static.a" Type="StaticLibrary" />
+    <BuildProduct Path="$(EngineDir)\Plugins\Runtime\PS4\Morpheus\Binaries\PS4\UE4Game-Morpheus-PS4-Shipping-Static.a" Type="StaticLibrary" />
     <BuildProduct Path="$(EngineDir)\Plugins\Runtime\PS4\PS4MoviePlayer\Binaries\PS4\UE4Game-PS4MoviePlayer-PS4-Shipping-Static.a" Type="StaticLibrary" />
     <BuildProduct Path="$(EngineDir)\Plugins\Slate\SlateRemote\Binaries\PS4\UE4Game-SlateRemote-PS4-Shipping-Static.a" Type="StaticLibrary" />
   </BuildProducts>
diff --git a/Engine/Build/Receipts/UE4Game-PS4-Test.target.xml b/Engine/Build/Receipts/UE4Game-PS4-Test.target.xml
--- a/Engine/Build/Receipts/UE4Game-PS4-Test.target.xml
+++ b/Engine/Build/Receipts/UE4Game-PS4-Test.target.xml
@@ -9,6 +9,7 @@
     <BuildProduct Path="$(EngineDir)\Plugins\Runtime\CableComponent\Binaries\PS4\UE4Game-CableComponent-PS4-Test-Static.a" Type="StaticLibrary" />
     <BuildProduct Path="$(EngineDir)\Plugins\Runtime\CustomMeshComponent\Binaries\PS4\UE4Game-CustomMeshComponent-PS4-Test-Static.a" Type="StaticLibrary" />
     <BuildProduct Path="$(EngineDir)\Plugins\Runtime\ExampleDeviceProfileSelector\Binaries\PS4\UE4Game-ExampleDeviceProfileSelector-PS4-Test-Static.a" Type="StaticLibrary" />
+    <BuildProduct Path="$(EngineDir)\Plugins\Runtime\PS4\Morpheus\Binaries\PS4\UE4Game-Morpheus-PS4-Test-Static.a" Type="StaticLibrary" />
     <BuildProduct Path="$(EngineDir)\Plugins\Runtime\PS4\PS4MoviePlayer\Binaries\PS4\UE4Game-PS4MoviePlayer-PS4-Test-Static.a" Type="StaticLibrary" />
     <BuildProduct Path="$(EngineDir)\Plugins\Slate\SlateRemote\Binaries\PS4\UE4Game-SlateRemote-PS4-Test-Static.a" Type="StaticLibrary" />
   </BuildProducts>
diff --git a/Engine/Build/Receipts/UE4Game-Win64-Test.target.xml b/Engine/Build/Receipts/UE4Game-Win64-Test.target.xml
--- a/Engine/Build/Receipts/UE4Game-Win64-Test.target.xml
+++ b/Engine/Build/Receipts/UE4Game-Win64-Test.target.xml
@@ -11,6 +11,7 @@
     <BuildProduct Path="$(EngineDir)\Plugins\Runtime\CustomMeshComponent\Binaries\Win64\UE4Game-CustomMeshComponent-Win64-Test-Static.lib" Type="StaticLibrary" />
     <BuildProduct Path="$(EngineDir)\Plugins\Runtime\ExampleDeviceProfileSelector\Binaries\Win64\UE4Game-ExampleDeviceProfileSelector-Win64-Test-Static.lib" Type="StaticLibrary" />
     <BuildProduct Path="$(EngineDir)\Plugins\Runtime\OculusRift\Binaries\Win64\UE4Game-OculusRift-Win64-Test-Static.lib" Type="StaticLibrary" />
+    <BuildProduct Path="$(EngineDir)\Plugins\Runtime\PS4\Morpheus\Binaries\Win64\UE4Game-Morpheus-Win64-Test-Static.lib" Type="StaticLibrary" />
     <BuildProduct Path="$(EngineDir)\Plugins\Runtime\WindowsMoviePlayer\Binaries\Win64\UE4Game-WindowsMoviePlayer-Win64-Test-Static.lib" Type="StaticLibrary" />
     <BuildProduct Path="$(EngineDir)\Plugins\Slate\SlateRemote\Binaries\Win64\UE4Game-SlateRemote-Win64-Test-Static.lib" Type="StaticLibrary" />
   </BuildProducts>
diff --git a/Engine/Build/Receipts/UE4Server-Win64-Development.target.xml b/Engine/Build/Receipts/UE4Server-Win64-Development.target.xml
--- a/Engine/Build/Receipts/UE4Server-Win64-Development.target.xml
+++ b/Engine/Build/Receipts/UE4Server-Win64-Development.target.xml
@@ -11,6 +11,7 @@
     <BuildProduct Path="$(EngineDir)\Plugins\Runtime\CustomMeshComponent\Binaries\Win64\UE4Server-CustomMeshComponent-Static.lib" Type="StaticLibrary" />
     <BuildProduct Path="$(EngineDir)\Plugins\Runtime\ExampleDeviceProfileSelector\Binaries\Win64\UE4Server-ExampleDeviceProfileSelector-Static.lib" Type="StaticLibrary" />
     <BuildProduct Path="$(EngineDir)\Plugins\Runtime\OculusRift\Binaries\Win64\UE4Server-OculusRift-Static.lib" Type="StaticLibrary" />
+    <BuildProduct Path="$(EngineDir)\Plugins\Runtime\PS4\Morpheus\Binaries\Win64\UE4Server-Morpheus-Static.lib" Type="StaticLibrary" />
     <BuildProduct Path="$(EngineDir)\Plugins\Runtime\WindowsMoviePlayer\Binaries\Win64\UE4Server-WindowsMoviePlayer-Static.lib" Type="StaticLibrary" />
     <BuildProduct Path="$(EngineDir)\Plugins\Slate\SlateRemote\Binaries\Win64\UE4Server-SlateRemote-Static.lib" Type="StaticLibrary" />
   </BuildProducts>
