GUBP Node Shared_LabelPromotable built from changelist 2507843

diff --git a/Engine/Build/Receipts/CrashReportClient-Linux-Shippingx86_64-unknown-linux-gnu.target.xml b/Engine/Build/Receipts/CrashReportClient-Linux-Shippingx86_64-unknown-linux-gnu.target.xml
--- a/Engine/Build/Receipts/CrashReportClient-Linux-Shippingx86_64-unknown-linux-gnu.target.xml
+++ b/Engine/Build/Receipts/CrashReportClient-Linux-Shippingx86_64-unknown-linux-gnu.target.xml
@@ -1,5 +1,11 @@
 <?xml version="1.0" encoding="utf-8"?>
 <BuildReceipt xmlns:xsd="http://www.w3.org/2001/XMLSchema" xmlns:xsi="http://www.w3.org/2001/XMLSchema-instance">
+  <Properties>
+    <Property Name="TargetName" Value="CrashReportClient" />
+    <Property Name="Platform" Value="Linux" />
+    <Property Name="Configuration" Value="Shipping" />
+    <Property Name="Precompile" Value="False" />
+  </Properties>
   <BuildProducts>
     <BuildProduct Path="$(EngineDir)\Binaries\Linux\CrashReportClient" Type="Executable" />
   </BuildProducts>
diff --git a/Engine/Build/Receipts/PS4MapFileUtil-Win64-Development.target.xml b/Engine/Build/Receipts/PS4MapFileUtil-Win64-Development.target.xml
--- a/Engine/Build/Receipts/PS4MapFileUtil-Win64-Development.target.xml
+++ b/Engine/Build/Receipts/PS4MapFileUtil-Win64-Development.target.xml
@@ -1,5 +1,11 @@
 <?xml version="1.0" encoding="utf-8"?>
 <BuildReceipt xmlns:xsd="http://www.w3.org/2001/XMLSchema" xmlns:xsi="http://www.w3.org/2001/XMLSchema-instance">
+  <Properties>
+    <Property Name="TargetName" Value="PS4MapFileUtil" />
+    <Property Name="Platform" Value="Win64" />
+    <Property Name="Configuration" Value="Development" />
+    <Property Name="Precompile" Value="False" />
+  </Properties>
   <BuildProducts>
     <BuildProduct Path="$(EngineDir)\Binaries\Win64\PS4\PS4MapFileUtil.exe" Type="Executable" />
     <BuildProduct Path="$(EngineDir)\Binaries\Win64\PS4\PS4MapFileUtil.pdb" Type="SymbolFile" />
diff --git a/Engine/Build/Receipts/ShaderCompileWorker-Linux-Developmentx86_64-unknown-linux-gnu.target.xml b/Engine/Build/Receipts/ShaderCompileWorker-Linux-Developmentx86_64-unknown-linux-gnu.target.xml
new file mode 100644
--- /dev/null
+++ b/Engine/Build/Receipts/ShaderCompileWorker-Linux-Developmentx86_64-unknown-linux-gnu.target.xml
@@ -0,0 +1,33 @@
+<?xml version="1.0" encoding="utf-8"?>
+<BuildReceipt xmlns:xsd="http://www.w3.org/2001/XMLSchema" xmlns:xsi="http://www.w3.org/2001/XMLSchema-instance">
+  <Properties>
+    <Property Name="TargetName" Value="ShaderCompileWorker" />
+    <Property Name="Platform" Value="Linux" />
+    <Property Name="Configuration" Value="Development" />
+    <Property Name="Precompile" Value="False" />
+  </Properties>
+  <BuildProducts>
+    <BuildProduct Path="$(EngineDir)\Binaries\Linux\ShaderCompileWorker" Type="Executable" />
+    <BuildProduct Path="$(EngineDir)\Binaries\Linux\libShaderCompileWorker-Core.so" Type="DynamicLibrary" />
+    <BuildProduct Path="$(EngineDir)\Binaries\Linux\libShaderCompileWorker-ShaderFormatOpenGL.so" Type="DynamicLibrary" />
+    <BuildProduct Path="$(EngineDir)\Binaries\Linux\libShaderCompileWorker-ShaderCore.so" Type="DynamicLibrary" />
+    <BuildProduct Path="$(EngineDir)\Binaries\Linux\libShaderCompileWorker-RHI.so" Type="DynamicLibrary" />
+    <BuildProduct Path="$(EngineDir)\Binaries\Linux\libShaderCompileWorker-RenderCore.so" Type="DynamicLibrary" />
+    <BuildProduct Path="$(EngineDir)\Binaries\Linux\libShaderCompileWorker-ShaderCompilerCommon.so" Type="DynamicLibrary" />
+    <BuildProduct Path="$(EngineDir)\Binaries\Linux\libShaderCompileWorker-ShaderPreprocessor.so" Type="DynamicLibrary" />
+    <BuildProduct Path="$(EngineDir)\Binaries\Linux\libShaderCompileWorker-Projects.so" Type="DynamicLibrary" />
+    <BuildProduct Path="$(EngineDir)\Binaries\Linux\libShaderCompileWorker-Json.so" Type="DynamicLibrary" />
+    <BuildProduct Path="$(EngineDir)\Binaries\Linux\libShaderCompileWorker-SandboxFile.so" Type="DynamicLibrary" />
+    <BuildProduct Path="$(EngineDir)\Binaries\Linux\libShaderCompileWorker-TargetPlatform.so" Type="DynamicLibrary" />
+    <BuildProduct Path="$(EngineDir)\Binaries\Linux\libShaderCompileWorker-DesktopPlatform.so" Type="DynamicLibrary" />
+    <BuildProduct Path="$(EngineDir)\Binaries\Linux\libShaderCompileWorker-NetworkFile.so" Type="DynamicLibrary" />
+    <BuildProduct Path="$(EngineDir)\Binaries\Linux\libShaderCompileWorker-CoreUObject.so" Type="DynamicLibrary" />
+    <BuildProduct Path="$(EngineDir)\Binaries\Linux\libShaderCompileWorker-Sockets.so" Type="DynamicLibrary" />
+    <BuildProduct Path="$(EngineDir)\Binaries\Linux\libShaderCompileWorker-DerivedDataCache.so" Type="DynamicLibrary" />
+    <BuildProduct Path="$(EngineDir)\Binaries\Linux\NotForLicensees\libShaderCompileWorker-DDCUtils.so" Type="DynamicLibrary" />
+    <BuildProduct Path="$(EngineDir)\Binaries\Linux\libShaderCompileWorker-PackageDependencyInfo.so" Type="DynamicLibrary" />
+    <BuildProduct Path="$(EngineDir)\Binaries\Linux\libShaderCompileWorker-PakFile.so" Type="DynamicLibrary" />
+    <BuildProduct Path="$(EngineDir)\Binaries\Linux\libShaderCompileWorker-StreamingFile.so" Type="DynamicLibrary" />
+  </BuildProducts>
+  <RuntimeDependencies />
+</BuildReceipt>
\ No newline at end of file
diff --git a/Engine/Build/Receipts/UE4Client-IOS-Development.target.xml b/Engine/Build/Receipts/UE4Client-IOS-Development.target.xml
--- a/Engine/Build/Receipts/UE4Client-IOS-Development.target.xml
+++ b/Engine/Build/Receipts/UE4Client-IOS-Development.target.xml
@@ -1,5 +1,11 @@
 <?xml version="1.0" encoding="utf-8"?>
 <BuildReceipt xmlns:xsi="http://www.w3.org/2001/XMLSchema-instance" xmlns:xsd="http://www.w3.org/2001/XMLSchema">
+  <Properties>
+    <Property Name="TargetName" Value="UE4Client" />
+    <Property Name="Platform" Value="IOS" />
+    <Property Name="Configuration" Value="Development" />
+    <Property Name="Precompile" Value="False" />
+  </Properties>
   <BuildProducts>
     <BuildProduct Path="$(EngineDir)/Binaries/IOS/UE4Client" Type="Executable" />
     <BuildProduct Path="$(EngineDir)/Binaries/IOS/UE4Client.stub" Type="Executable" />
diff --git a/Engine/Build/Receipts/UE4Client-Linux-Developmentx86_64-unknown-linux-gnu.target.xml b/Engine/Build/Receipts/UE4Client-Linux-Developmentx86_64-unknown-linux-gnu.target.xml
--- a/Engine/Build/Receipts/UE4Client-Linux-Developmentx86_64-unknown-linux-gnu.target.xml
+++ b/Engine/Build/Receipts/UE4Client-Linux-Developmentx86_64-unknown-linux-gnu.target.xml
@@ -1,5 +1,11 @@
 <?xml version="1.0" encoding="utf-8"?>
 <BuildReceipt xmlns:xsd="http://www.w3.org/2001/XMLSchema" xmlns:xsi="http://www.w3.org/2001/XMLSchema-instance">
+  <Properties>
+    <Property Name="TargetName" Value="UE4Client" />
+    <Property Name="Platform" Value="Linux" />
+    <Property Name="Configuration" Value="Development" />
+    <Property Name="Precompile" Value="False" />
+  </Properties>
   <BuildProducts>
     <BuildProduct Path="$(EngineDir)\Binaries\Linux\UE4Client" Type="Executable" />
     <BuildProduct Path="$(EngineDir)\Plugins\2D\Paper2D\Binaries\Linux\libUE4Client-Paper2D-Static.a" Type="StaticLibrary" />
diff --git a/Engine/Build/Receipts/UE4Client-Mac-Development.target.xml b/Engine/Build/Receipts/UE4Client-Mac-Development.target.xml
--- a/Engine/Build/Receipts/UE4Client-Mac-Development.target.xml
+++ b/Engine/Build/Receipts/UE4Client-Mac-Development.target.xml
@@ -1,5 +1,11 @@
 <?xml version="1.0" encoding="utf-8"?>
 <BuildReceipt xmlns:xsi="http://www.w3.org/2001/XMLSchema-instance" xmlns:xsd="http://www.w3.org/2001/XMLSchema">
+  <Properties>
+    <Property Name="TargetName" Value="UE4Client" />
+    <Property Name="Platform" Value="Mac" />
+    <Property Name="Configuration" Value="Development" />
+    <Property Name="Precompile" Value="False" />
+  </Properties>
   <BuildProducts>
     <BuildProduct Path="$(EngineDir)/Binaries/Mac/UE4Client.app/Contents/MacOS/UE4Client" Type="Executable" />
     <BuildProduct Path="$(EngineDir)/Binaries/Mac/UE4Client.app/Contents/MacOS/libogg.dylib" Type="DynamicLibrary" />
diff --git a/Engine/Build/Receipts/UE4Client-Win32-Development.target.xml b/Engine/Build/Receipts/UE4Client-Win32-Development.target.xml
--- a/Engine/Build/Receipts/UE4Client-Win32-Development.target.xml
+++ b/Engine/Build/Receipts/UE4Client-Win32-Development.target.xml
@@ -1,5 +1,11 @@
 <?xml version="1.0" encoding="utf-8"?>
 <BuildReceipt xmlns:xsd="http://www.w3.org/2001/XMLSchema" xmlns:xsi="http://www.w3.org/2001/XMLSchema-instance">
+  <Properties>
+    <Property Name="TargetName" Value="UE4Client" />
+    <Property Name="Platform" Value="Win32" />
+    <Property Name="Configuration" Value="Development" />
+    <Property Name="Precompile" Value="False" />
+  </Properties>
   <BuildProducts>
     <BuildProduct Path="$(EngineDir)\Binaries\Win32\UE4Client.exe" Type="Executable" />
     <BuildProduct Path="$(EngineDir)\Binaries\Win32\UE4Client.pdb" Type="SymbolFile" />
diff --git a/Engine/Build/Receipts/UE4Client-Win64-Development.target.xml b/Engine/Build/Receipts/UE4Client-Win64-Development.target.xml
--- a/Engine/Build/Receipts/UE4Client-Win64-Development.target.xml
+++ b/Engine/Build/Receipts/UE4Client-Win64-Development.target.xml
@@ -1,5 +1,11 @@
 <?xml version="1.0" encoding="utf-8"?>
 <BuildReceipt xmlns:xsd="http://www.w3.org/2001/XMLSchema" xmlns:xsi="http://www.w3.org/2001/XMLSchema-instance">
+  <Properties>
+    <Property Name="TargetName" Value="UE4Client" />
+    <Property Name="Platform" Value="Win64" />
+    <Property Name="Configuration" Value="Development" />
+    <Property Name="Precompile" Value="False" />
+  </Properties>
   <BuildProducts>
     <BuildProduct Path="$(EngineDir)\Binaries\Win64\UE4Client.exe" Type="Executable" />
     <BuildProduct Path="$(EngineDir)\Binaries\Win64\UE4Client.pdb" Type="SymbolFile" />
@@ -11,7 +17,6 @@
     <BuildProduct Path="$(EngineDir)\Plugins\Runtime\CustomMeshComponent\Binaries\Win64\UE4Client-CustomMeshComponent-Static.lib" Type="StaticLibrary" />
     <BuildProduct Path="$(EngineDir)\Plugins\Runtime\ExampleDeviceProfileSelector\Binaries\Win64\UE4Client-ExampleDeviceProfileSelector-Static.lib" Type="StaticLibrary" />
     <BuildProduct Path="$(EngineDir)\Plugins\Runtime\OculusRift\Binaries\Win64\UE4Client-OculusRift-Static.lib" Type="StaticLibrary" />
-    <BuildProduct Path="$(EngineDir)\Plugins\Runtime\PS4\Morpheus\Binaries\Win64\UE4Client-Morpheus-Static.lib" Type="StaticLibrary" />
     <BuildProduct Path="$(EngineDir)\Plugins\Runtime\WindowsMoviePlayer\Binaries\Win64\UE4Client-WindowsMoviePlayer-Static.lib" Type="StaticLibrary" />
     <BuildProduct Path="$(EngineDir)\Plugins\Slate\SlateRemote\Binaries\Win64\UE4Client-SlateRemote-Static.lib" Type="StaticLibrary" />
   </BuildProducts>
diff --git a/Engine/Build/Receipts/UE4Editor-Linux-Developmentx86_64-unknown-linux-gnu.target.xml b/Engine/Build/Receipts/UE4Editor-Linux-Developmentx86_64-unknown-linux-gnu.target.xml
--- a/Engine/Build/Receipts/UE4Editor-Linux-Developmentx86_64-unknown-linux-gnu.target.xml
+++ b/Engine/Build/Receipts/UE4Editor-Linux-Developmentx86_64-unknown-linux-gnu.target.xml
@@ -1,5 +1,11 @@
 <?xml version="1.0" encoding="utf-8"?>
 <BuildReceipt xmlns:xsd="http://www.w3.org/2001/XMLSchema" xmlns:xsi="http://www.w3.org/2001/XMLSchema-instance">
+  <Properties>
+    <Property Name="TargetName" Value="UE4Editor" />
+    <Property Name="Platform" Value="Linux" />
+    <Property Name="Configuration" Value="Development" />
+    <Property Name="Precompile" Value="False" />
+  </Properties>
   <BuildProducts>
     <BuildProduct Path="$(EngineDir)\Binaries\Linux\UE4Editor" Type="Executable" />
     <BuildProduct Path="$(EngineDir)\Plugins\2D\Paper2D\Binaries\Linux\libUE4Editor-Paper2D.so" Type="DynamicLibrary" />
diff --git a/Engine/Build/Receipts/UE4Game-Android-Development.target.xml b/Engine/Build/Receipts/UE4Game-Android-Development.target.xml
--- a/Engine/Build/Receipts/UE4Game-Android-Development.target.xml
+++ b/Engine/Build/Receipts/UE4Game-Android-Development.target.xml
@@ -1,5 +1,11 @@
 <?xml version="1.0" encoding="utf-8"?>
 <BuildReceipt xmlns:xsd="http://www.w3.org/2001/XMLSchema" xmlns:xsi="http://www.w3.org/2001/XMLSchema-instance">
+  <Properties>
+    <Property Name="TargetName" Value="UE4Game" />
+    <Property Name="Platform" Value="Android" />
+    <Property Name="Configuration" Value="Development" />
+    <Property Name="Precompile" Value="True" />
+  </Properties>
   <BuildProducts>
     <BuildProduct Path="$(EngineDir)\Binaries\Android\UE4Game-armv7-es2.so" Type="Executable" />
     <BuildProduct Path="$(EngineDir)\Binaries\Android\UE4Game-armv7-es2.apk" Type="Executable" />
diff --git a/Engine/Build/Receipts/UE4Game-Android-Shipping.target.xml b/Engine/Build/Receipts/UE4Game-Android-Shipping.target.xml
--- a/Engine/Build/Receipts/UE4Game-Android-Shipping.target.xml
+++ b/Engine/Build/Receipts/UE4Game-Android-Shipping.target.xml
@@ -1,5 +1,11 @@
 <?xml version="1.0" encoding="utf-8"?>
 <BuildReceipt xmlns:xsd="http://www.w3.org/2001/XMLSchema" xmlns:xsi="http://www.w3.org/2001/XMLSchema-instance">
+  <Properties>
+    <Property Name="TargetName" Value="UE4Game" />
+    <Property Name="Platform" Value="Android" />
+    <Property Name="Configuration" Value="Shipping" />
+    <Property Name="Precompile" Value="True" />
+  </Properties>
   <BuildProducts>
     <BuildProduct Path="$(EngineDir)\Binaries\Android\UE4Game-Android-Shipping-armv7-es2.so" Type="Executable" />
     <BuildProduct Path="$(EngineDir)\Binaries\Android\UE4Game-Android-Shipping-armv7-es2.apk" Type="Executable" />
diff --git a/Engine/Build/Receipts/UE4Game-Android-Test.target.xml b/Engine/Build/Receipts/UE4Game-Android-Test.target.xml
--- a/Engine/Build/Receipts/UE4Game-Android-Test.target.xml
+++ b/Engine/Build/Receipts/UE4Game-Android-Test.target.xml
@@ -1,5 +1,11 @@
 <?xml version="1.0" encoding="utf-8"?>
 <BuildReceipt xmlns:xsd="http://www.w3.org/2001/XMLSchema" xmlns:xsi="http://www.w3.org/2001/XMLSchema-instance">
+  <Properties>
+    <Property Name="TargetName" Value="UE4Game" />
+    <Property Name="Platform" Value="Android" />
+    <Property Name="Configuration" Value="Test" />
+    <Property Name="Precompile" Value="False" />
+  </Properties>
   <BuildProducts>
     <BuildProduct Path="$(EngineDir)\Binaries\Android\UE4Game-Android-Test-armv7-es2.so" Type="Executable" />
     <BuildProduct Path="$(EngineDir)\Binaries\Android\UE4Game-Android-Test-armv7-es2.apk" Type="Executable" />
diff --git a/Engine/Build/Receipts/UE4Game-HTML5-Development.target.xml b/Engine/Build/Receipts/UE4Game-HTML5-Development.target.xml
--- a/Engine/Build/Receipts/UE4Game-HTML5-Development.target.xml
+++ b/Engine/Build/Receipts/UE4Game-HTML5-Development.target.xml
@@ -1,5 +1,11 @@
 <?xml version="1.0" encoding="utf-8"?>
 <BuildReceipt xmlns:xsd="http://www.w3.org/2001/XMLSchema" xmlns:xsi="http://www.w3.org/2001/XMLSchema-instance">
+  <Properties>
+    <Property Name="TargetName" Value="UE4Game" />
+    <Property Name="Platform" Value="HTML5" />
+    <Property Name="Configuration" Value="Development" />
+    <Property Name="Precompile" Value="True" />
+  </Properties>
   <BuildProducts>
     <BuildProduct Path="$(EngineDir)\Binaries\HTML5\UE4Game.js" Type="Executable" />
     <BuildProduct Path="$(EngineDir)\Binaries\HTML5\UE4Game.js.mem" Type="RequiredResource" />
diff --git a/Engine/Build/Receipts/UE4Game-HTML5-Shipping.target.xml b/Engine/Build/Receipts/UE4Game-HTML5-Shipping.target.xml
--- a/Engine/Build/Receipts/UE4Game-HTML5-Shipping.target.xml
+++ b/Engine/Build/Receipts/UE4Game-HTML5-Shipping.target.xml
@@ -1,5 +1,11 @@
 <?xml version="1.0" encoding="utf-8"?>
 <BuildReceipt xmlns:xsd="http://www.w3.org/2001/XMLSchema" xmlns:xsi="http://www.w3.org/2001/XMLSchema-instance">
+  <Properties>
+    <Property Name="TargetName" Value="UE4Game" />
+    <Property Name="Platform" Value="HTML5" />
+    <Property Name="Configuration" Value="Shipping" />
+    <Property Name="Precompile" Value="True" />
+  </Properties>
   <BuildProducts>
     <BuildProduct Path="$(EngineDir)\Binaries\HTML5\UE4Game-HTML5-Shipping.js" Type="Executable" />
     <BuildProduct Path="$(EngineDir)\Binaries\HTML5\UE4Game-HTML5-Shipping.js.mem" Type="RequiredResource" />
diff --git a/Engine/Build/Receipts/UE4Game-HTML5-Test.target.xml b/Engine/Build/Receipts/UE4Game-HTML5-Test.target.xml
--- a/Engine/Build/Receipts/UE4Game-HTML5-Test.target.xml
+++ b/Engine/Build/Receipts/UE4Game-HTML5-Test.target.xml
@@ -1,5 +1,11 @@
 <?xml version="1.0" encoding="utf-8"?>
 <BuildReceipt xmlns:xsd="http://www.w3.org/2001/XMLSchema" xmlns:xsi="http://www.w3.org/2001/XMLSchema-instance">
+  <Properties>
+    <Property Name="TargetName" Value="UE4Game" />
+    <Property Name="Platform" Value="HTML5" />
+    <Property Name="Configuration" Value="Test" />
+    <Property Name="Precompile" Value="False" />
+  </Properties>
   <BuildProducts>
     <BuildProduct Path="$(EngineDir)\Binaries\HTML5\UE4Game-HTML5-Test.js" Type="Executable" />
     <BuildProduct Path="$(EngineDir)\Binaries\HTML5\UE4Game-HTML5-Test.js.mem" Type="RequiredResource" />
diff --git a/Engine/Build/Receipts/UE4Game-IOS-Development.target.xml b/Engine/Build/Receipts/UE4Game-IOS-Development.target.xml
--- a/Engine/Build/Receipts/UE4Game-IOS-Development.target.xml
+++ b/Engine/Build/Receipts/UE4Game-IOS-Development.target.xml
@@ -1,5 +1,11 @@
 <?xml version="1.0" encoding="utf-8"?>
 <BuildReceipt xmlns:xsi="http://www.w3.org/2001/XMLSchema-instance" xmlns:xsd="http://www.w3.org/2001/XMLSchema">
+  <Properties>
+    <Property Name="TargetName" Value="UE4Game" />
+    <Property Name="Platform" Value="IOS" />
+    <Property Name="Configuration" Value="Development" />
+    <Property Name="Precompile" Value="True" />
+  </Properties>
   <BuildProducts>
     <BuildProduct Path="$(EngineDir)/Binaries/IOS/UE4Game" Type="Executable" />
     <BuildProduct Path="$(EngineDir)/Binaries/IOS/UE4Game.stub" Type="Executable" />
diff --git a/Engine/Build/Receipts/UE4Game-IOS-Shipping.target.xml b/Engine/Build/Receipts/UE4Game-IOS-Shipping.target.xml
--- a/Engine/Build/Receipts/UE4Game-IOS-Shipping.target.xml
+++ b/Engine/Build/Receipts/UE4Game-IOS-Shipping.target.xml
@@ -1,5 +1,11 @@
 <?xml version="1.0" encoding="utf-8"?>
 <BuildReceipt xmlns:xsi="http://www.w3.org/2001/XMLSchema-instance" xmlns:xsd="http://www.w3.org/2001/XMLSchema">
+  <Properties>
+    <Property Name="TargetName" Value="UE4Game" />
+    <Property Name="Platform" Value="IOS" />
+    <Property Name="Configuration" Value="Shipping" />
+    <Property Name="Precompile" Value="True" />
+  </Properties>
   <BuildProducts>
     <BuildProduct Path="$(EngineDir)/Binaries/IOS/UE4Game-IOS-Shipping" Type="Executable" />
     <BuildProduct Path="$(EngineDir)/Binaries/IOS/UE4Game-IOS-Shipping.stub" Type="Executable" />
diff --git a/Engine/Build/Receipts/UE4Game-IOS-Test.target.xml b/Engine/Build/Receipts/UE4Game-IOS-Test.target.xml
--- a/Engine/Build/Receipts/UE4Game-IOS-Test.target.xml
+++ b/Engine/Build/Receipts/UE4Game-IOS-Test.target.xml
@@ -1,5 +1,11 @@
 <?xml version="1.0" encoding="utf-8"?>
 <BuildReceipt xmlns:xsi="http://www.w3.org/2001/XMLSchema-instance" xmlns:xsd="http://www.w3.org/2001/XMLSchema">
+  <Properties>
+    <Property Name="TargetName" Value="UE4Game" />
+    <Property Name="Platform" Value="IOS" />
+    <Property Name="Configuration" Value="Test" />
+    <Property Name="Precompile" Value="False" />
+  </Properties>
   <BuildProducts>
     <BuildProduct Path="$(EngineDir)/Binaries/IOS/UE4Game-IOS-Test" Type="Executable" />
     <BuildProduct Path="$(EngineDir)/Binaries/IOS/UE4Game-IOS-Test.stub" Type="Executable" />
diff --git a/Engine/Build/Receipts/UE4Game-Linux-Developmentx86_64-unknown-linux-gnu.target.xml b/Engine/Build/Receipts/UE4Game-Linux-Developmentx86_64-unknown-linux-gnu.target.xml
--- a/Engine/Build/Receipts/UE4Game-Linux-Developmentx86_64-unknown-linux-gnu.target.xml
+++ b/Engine/Build/Receipts/UE4Game-Linux-Developmentx86_64-unknown-linux-gnu.target.xml
@@ -1,5 +1,11 @@
 <?xml version="1.0" encoding="utf-8"?>
 <BuildReceipt xmlns:xsd="http://www.w3.org/2001/XMLSchema" xmlns:xsi="http://www.w3.org/2001/XMLSchema-instance">
+  <Properties>
+    <Property Name="TargetName" Value="UE4Game" />
+    <Property Name="Platform" Value="Linux" />
+    <Property Name="Configuration" Value="Development" />
+    <Property Name="Precompile" Value="False" />
+  </Properties>
   <BuildProducts>
     <BuildProduct Path="$(EngineDir)\Binaries\Linux\UE4Game" Type="Executable" />
     <BuildProduct Path="$(EngineDir)\Plugins\2D\Paper2D\Binaries\Linux\libUE4Game-Paper2D-Static.a" Type="StaticLibrary" />
diff --git a/Engine/Build/Receipts/UE4Game-Linux-Shippingx86_64-unknown-linux-gnu.target.xml b/Engine/Build/Receipts/UE4Game-Linux-Shippingx86_64-unknown-linux-gnu.target.xml
--- a/Engine/Build/Receipts/UE4Game-Linux-Shippingx86_64-unknown-linux-gnu.target.xml
+++ b/Engine/Build/Receipts/UE4Game-Linux-Shippingx86_64-unknown-linux-gnu.target.xml
@@ -1,5 +1,11 @@
 <?xml version="1.0" encoding="utf-8"?>
 <BuildReceipt xmlns:xsd="http://www.w3.org/2001/XMLSchema" xmlns:xsi="http://www.w3.org/2001/XMLSchema-instance">
+  <Properties>
+    <Property Name="TargetName" Value="UE4Game" />
+    <Property Name="Platform" Value="Linux" />
+    <Property Name="Configuration" Value="Shipping" />
+    <Property Name="Precompile" Value="False" />
+  </Properties>
   <BuildProducts>
     <BuildProduct Path="$(EngineDir)\Binaries\Linux\UE4Game-Linux-Shipping" Type="Executable" />
     <BuildProduct Path="$(EngineDir)\Plugins\2D\Paper2D\Binaries\Linux\libUE4Game-Paper2D-Linux-Shipping-Static.a" Type="StaticLibrary" />
diff --git a/Engine/Build/Receipts/UE4Game-Linux-Testx86_64-unknown-linux-gnu.target.xml b/Engine/Build/Receipts/UE4Game-Linux-Testx86_64-unknown-linux-gnu.target.xml
--- a/Engine/Build/Receipts/UE4Game-Linux-Testx86_64-unknown-linux-gnu.target.xml
+++ b/Engine/Build/Receipts/UE4Game-Linux-Testx86_64-unknown-linux-gnu.target.xml
@@ -1,5 +1,11 @@
 <?xml version="1.0" encoding="utf-8"?>
 <BuildReceipt xmlns:xsd="http://www.w3.org/2001/XMLSchema" xmlns:xsi="http://www.w3.org/2001/XMLSchema-instance">
+  <Properties>
+    <Property Name="TargetName" Value="UE4Game" />
+    <Property Name="Platform" Value="Linux" />
+    <Property Name="Configuration" Value="Test" />
+    <Property Name="Precompile" Value="False" />
+  </Properties>
   <BuildProducts>
     <BuildProduct Path="$(EngineDir)\Binaries\Linux\UE4Game-Linux-Test" Type="Executable" />
     <BuildProduct Path="$(EngineDir)\Plugins\2D\Paper2D\Binaries\Linux\libUE4Game-Paper2D-Linux-Test-Static.a" Type="StaticLibrary" />
diff --git a/Engine/Build/Receipts/UE4Game-Mac-Development.target.xml b/Engine/Build/Receipts/UE4Game-Mac-Development.target.xml
--- a/Engine/Build/Receipts/UE4Game-Mac-Development.target.xml
+++ b/Engine/Build/Receipts/UE4Game-Mac-Development.target.xml
@@ -1,5 +1,11 @@
 <?xml version="1.0" encoding="utf-8"?>
 <BuildReceipt xmlns:xsi="http://www.w3.org/2001/XMLSchema-instance" xmlns:xsd="http://www.w3.org/2001/XMLSchema">
+  <Properties>
+    <Property Name="TargetName" Value="UE4Game" />
+    <Property Name="Platform" Value="Mac" />
+    <Property Name="Configuration" Value="Development" />
+    <Property Name="Precompile" Value="True" />
+  </Properties>
   <BuildProducts>
     <BuildProduct Path="$(EngineDir)/Binaries/Mac/UE4Game.app/Contents/MacOS/UE4Game" Type="Executable" />
     <BuildProduct Path="$(EngineDir)/Binaries/Mac/UE4Game.app/Contents/MacOS/libogg.dylib" Type="DynamicLibrary" />
diff --git a/Engine/Build/Receipts/UE4Game-Mac-Shipping.target.xml b/Engine/Build/Receipts/UE4Game-Mac-Shipping.target.xml
--- a/Engine/Build/Receipts/UE4Game-Mac-Shipping.target.xml
+++ b/Engine/Build/Receipts/UE4Game-Mac-Shipping.target.xml
@@ -1,5 +1,11 @@
 <?xml version="1.0" encoding="utf-8"?>
 <BuildReceipt xmlns:xsi="http://www.w3.org/2001/XMLSchema-instance" xmlns:xsd="http://www.w3.org/2001/XMLSchema">
+  <Properties>
+    <Property Name="TargetName" Value="UE4Game" />
+    <Property Name="Platform" Value="Mac" />
+    <Property Name="Configuration" Value="Shipping" />
+    <Property Name="Precompile" Value="True" />
+  </Properties>
   <BuildProducts>
     <BuildProduct Path="$(EngineDir)/Binaries/Mac/UE4Game-Mac-Shipping.app/Contents/MacOS/UE4Game-Mac-Shipping" Type="Executable" />
     <BuildProduct Path="$(EngineDir)/Binaries/Mac/UE4Game-Mac-Shipping.app/Contents/MacOS/libogg.dylib" Type="DynamicLibrary" />
diff --git a/Engine/Build/Receipts/UE4Game-Mac-Test.target.xml b/Engine/Build/Receipts/UE4Game-Mac-Test.target.xml
--- a/Engine/Build/Receipts/UE4Game-Mac-Test.target.xml
+++ b/Engine/Build/Receipts/UE4Game-Mac-Test.target.xml
@@ -1,5 +1,11 @@
 <?xml version="1.0" encoding="utf-8"?>
 <BuildReceipt xmlns:xsi="http://www.w3.org/2001/XMLSchema-instance" xmlns:xsd="http://www.w3.org/2001/XMLSchema">
+  <Properties>
+    <Property Name="TargetName" Value="UE4Game" />
+    <Property Name="Platform" Value="Mac" />
+    <Property Name="Configuration" Value="Test" />
+    <Property Name="Precompile" Value="False" />
+  </Properties>
   <BuildProducts>
     <BuildProduct Path="$(EngineDir)/Binaries/Mac/UE4Game-Mac-Test.app/Contents/MacOS/UE4Game-Mac-Test" Type="Executable" />
     <BuildProduct Path="$(EngineDir)/Binaries/Mac/UE4Game-Mac-Test.app/Contents/MacOS/libogg.dylib" Type="DynamicLibrary" />
diff --git a/Engine/Build/Receipts/UE4Game-PS4-Development.target.xml b/Engine/Build/Receipts/UE4Game-PS4-Development.target.xml
--- a/Engine/Build/Receipts/UE4Game-PS4-Development.target.xml
+++ b/Engine/Build/Receipts/UE4Game-PS4-Development.target.xml
@@ -1,5 +1,11 @@
 <?xml version="1.0" encoding="utf-8"?>
 <BuildReceipt xmlns:xsd="http://www.w3.org/2001/XMLSchema" xmlns:xsi="http://www.w3.org/2001/XMLSchema-instance">
+  <Properties>
+    <Property Name="TargetName" Value="UE4Game" />
+    <Property Name="Platform" Value="PS4" />
+    <Property Name="Configuration" Value="Development" />
+    <Property Name="Precompile" Value="False" />
+  </Properties>
   <BuildProducts>
     <BuildProduct Path="$(EngineDir)\Binaries\PS4\UE4Game.self" Type="Executable" />
     <BuildProduct Path="$(EngineDir)\Plugins\2D\Paper2D\Binaries\PS4\UE4Game-Paper2D-Static.a" Type="StaticLibrary" />
@@ -9,7 +15,6 @@
     <BuildProduct Path="$(EngineDir)\Plugins\Runtime\CableComponent\Binaries\PS4\UE4Game-CableComponent-Static.a" Type="StaticLibrary" />
     <BuildProduct Path="$(EngineDir)\Plugins\Runtime\CustomMeshComponent\Binaries\PS4\UE4Game-CustomMeshComponent-Static.a" Type="StaticLibrary" />
     <BuildProduct Path="$(EngineDir)\Plugins\Runtime\ExampleDeviceProfileSelector\Binaries\PS4\UE4Game-ExampleDeviceProfileSelector-Static.a" Type="StaticLibrary" />
-    <BuildProduct Path="$(EngineDir)\Plugins\Runtime\PS4\Morpheus\Binaries\PS4\UE4Game-Morpheus-Static.a" Type="StaticLibrary" />
     <BuildProduct Path="$(EngineDir)\Plugins\Runtime\PS4\PS4MoviePlayer\Binaries\PS4\UE4Game-PS4MoviePlayer-Static.a" Type="StaticLibrary" />
     <BuildProduct Path="$(EngineDir)\Plugins\Slate\SlateRemote\Binaries\PS4\UE4Game-SlateRemote-Static.a" Type="StaticLibrary" />
   </BuildProducts>
diff --git a/Engine/Build/Receipts/UE4Game-PS4-Shipping.target.xml b/Engine/Build/Receipts/UE4Game-PS4-Shipping.target.xml
--- a/Engine/Build/Receipts/UE4Game-PS4-Shipping.target.xml
+++ b/Engine/Build/Receipts/UE4Game-PS4-Shipping.target.xml
@@ -1,5 +1,11 @@
 <?xml version="1.0" encoding="utf-8"?>
 <BuildReceipt xmlns:xsd="http://www.w3.org/2001/XMLSchema" xmlns:xsi="http://www.w3.org/2001/XMLSchema-instance">
+  <Properties>
+    <Property Name="TargetName" Value="UE4Game" />
+    <Property Name="Platform" Value="PS4" />
+    <Property Name="Configuration" Value="Shipping" />
+    <Property Name="Precompile" Value="False" />
+  </Properties>
   <BuildProducts>
     <BuildProduct Path="$(EngineDir)\Binaries\PS4\UE4Game-PS4-Shipping.self" Type="Executable" />
     <BuildProduct Path="$(EngineDir)\Plugins\2D\Paper2D\Binaries\PS4\UE4Game-Paper2D-PS4-Shipping-Static.a" Type="StaticLibrary" />
@@ -9,7 +15,6 @@
     <BuildProduct Path="$(EngineDir)\Plugins\Runtime\CableComponent\Binaries\PS4\UE4Game-CableComponent-PS4-Shipping-Static.a" Type="StaticLibrary" />
     <BuildProduct Path="$(EngineDir)\Plugins\Runtime\CustomMeshComponent\Binaries\PS4\UE4Game-CustomMeshComponent-PS4-Shipping-Static.a" Type="StaticLibrary" />
     <BuildProduct Path="$(EngineDir)\Plugins\Runtime\ExampleDeviceProfileSelector\Binaries\PS4\UE4Game-ExampleDeviceProfileSelector-PS4-Shipping-Static.a" Type="StaticLibrary" />
-    <BuildProduct Path="$(EngineDir)\Plugins\Runtime\PS4\Morpheus\Binaries\PS4\UE4Game-Morpheus-PS4-Shipping-Static.a" Type="StaticLibrary" />
     <BuildProduct Path="$(EngineDir)\Plugins\Runtime\PS4\PS4MoviePlayer\Binaries\PS4\UE4Game-PS4MoviePlayer-PS4-Shipping-Static.a" Type="StaticLibrary" />
     <BuildProduct Path="$(EngineDir)\Plugins\Slate\SlateRemote\Binaries\PS4\UE4Game-SlateRemote-PS4-Shipping-Static.a" Type="StaticLibrary" />
   </BuildProducts>
diff --git a/Engine/Build/Receipts/UE4Game-PS4-Test.target.xml b/Engine/Build/Receipts/UE4Game-PS4-Test.target.xml
--- a/Engine/Build/Receipts/UE4Game-PS4-Test.target.xml
+++ b/Engine/Build/Receipts/UE4Game-PS4-Test.target.xml
@@ -1,5 +1,11 @@
 <?xml version="1.0" encoding="utf-8"?>
 <BuildReceipt xmlns:xsd="http://www.w3.org/2001/XMLSchema" xmlns:xsi="http://www.w3.org/2001/XMLSchema-instance">
+  <Properties>
+    <Property Name="TargetName" Value="UE4Game" />
+    <Property Name="Platform" Value="PS4" />
+    <Property Name="Configuration" Value="Test" />
+    <Property Name="Precompile" Value="False" />
+  </Properties>
   <BuildProducts>
     <BuildProduct Path="$(EngineDir)\Binaries\PS4\UE4Game-PS4-Test.self" Type="Executable" />
     <BuildProduct Path="$(EngineDir)\Plugins\2D\Paper2D\Binaries\PS4\UE4Game-Paper2D-PS4-Test-Static.a" Type="StaticLibrary" />
@@ -9,7 +15,6 @@
     <BuildProduct Path="$(EngineDir)\Plugins\Runtime\CableComponent\Binaries\PS4\UE4Game-CableComponent-PS4-Test-Static.a" Type="StaticLibrary" />
     <BuildProduct Path="$(EngineDir)\Plugins\Runtime\CustomMeshComponent\Binaries\PS4\UE4Game-CustomMeshComponent-PS4-Test-Static.a" Type="StaticLibrary" />
     <BuildProduct Path="$(EngineDir)\Plugins\Runtime\ExampleDeviceProfileSelector\Binaries\PS4\UE4Game-ExampleDeviceProfileSelector-PS4-Test-Static.a" Type="StaticLibrary" />
-    <BuildProduct Path="$(EngineDir)\Plugins\Runtime\PS4\Morpheus\Binaries\PS4\UE4Game-Morpheus-PS4-Test-Static.a" Type="StaticLibrary" />
     <BuildProduct Path="$(EngineDir)\Plugins\Runtime\PS4\PS4MoviePlayer\Binaries\PS4\UE4Game-PS4MoviePlayer-PS4-Test-Static.a" Type="StaticLibrary" />
     <BuildProduct Path="$(EngineDir)\Plugins\Slate\SlateRemote\Binaries\PS4\UE4Game-SlateRemote-PS4-Test-Static.a" Type="StaticLibrary" />
   </BuildProducts>
diff --git a/Engine/Build/Receipts/UE4Game-Win32-Development.target.xml b/Engine/Build/Receipts/UE4Game-Win32-Development.target.xml
--- a/Engine/Build/Receipts/UE4Game-Win32-Development.target.xml
+++ b/Engine/Build/Receipts/UE4Game-Win32-Development.target.xml
@@ -1,5 +1,11 @@
 <?xml version="1.0" encoding="utf-8"?>
 <BuildReceipt xmlns:xsd="http://www.w3.org/2001/XMLSchema" xmlns:xsi="http://www.w3.org/2001/XMLSchema-instance">
+  <Properties>
+    <Property Name="TargetName" Value="UE4Game" />
+    <Property Name="Platform" Value="Win32" />
+    <Property Name="Configuration" Value="Development" />
+    <Property Name="Precompile" Value="True" />
+  </Properties>
   <BuildProducts>
     <BuildProduct Path="$(EngineDir)\Binaries\Win32\UE4Game.exe" Type="Executable" />
     <BuildProduct Path="$(EngineDir)\Binaries\Win32\UE4Game.pdb" Type="SymbolFile" />
diff --git a/Engine/Build/Receipts/UE4Game-Win32-Development_XP.target.xml b/Engine/Build/Receipts/UE4Game-Win32-Development_XP.target.xml
--- a/Engine/Build/Receipts/UE4Game-Win32-Development_XP.target.xml
+++ b/Engine/Build/Receipts/UE4Game-Win32-Development_XP.target.xml
@@ -1,5 +1,11 @@
 <?xml version="1.0" encoding="utf-8"?>
 <BuildReceipt xmlns:xsd="http://www.w3.org/2001/XMLSchema" xmlns:xsi="http://www.w3.org/2001/XMLSchema-instance">
+  <Properties>
+    <Property Name="TargetName" Value="UE4Game" />
+    <Property Name="Platform" Value="Win32" />
+    <Property Name="Configuration" Value="Development" />
+    <Property Name="Precompile" Value="False" />
+  </Properties>
   <BuildProducts>
     <BuildProduct Path="$(EngineDir)\Binaries\Win32\UE4Game_XP.exe" Type="Executable" />
     <BuildProduct Path="$(EngineDir)\Binaries\Win32\UE4Game_XP.pdb" Type="SymbolFile" />
diff --git a/Engine/Build/Receipts/UE4Game-Win32-Shipping.target.xml b/Engine/Build/Receipts/UE4Game-Win32-Shipping.target.xml
--- a/Engine/Build/Receipts/UE4Game-Win32-Shipping.target.xml
+++ b/Engine/Build/Receipts/UE4Game-Win32-Shipping.target.xml
@@ -1,5 +1,11 @@
 <?xml version="1.0" encoding="utf-8"?>
 <BuildReceipt xmlns:xsd="http://www.w3.org/2001/XMLSchema" xmlns:xsi="http://www.w3.org/2001/XMLSchema-instance">
+  <Properties>
+    <Property Name="TargetName" Value="UE4Game" />
+    <Property Name="Platform" Value="Win32" />
+    <Property Name="Configuration" Value="Shipping" />
+    <Property Name="Precompile" Value="True" />
+  </Properties>
   <BuildProducts>
     <BuildProduct Path="$(EngineDir)\Binaries\Win32\UE4Game-Win32-Shipping.exe" Type="Executable" />
     <BuildProduct Path="$(EngineDir)\Binaries\Win32\UE4Game-Win32-Shipping.pdb" Type="SymbolFile" />
diff --git a/Engine/Build/Receipts/UE4Game-Win32-Shipping_XP.target.xml b/Engine/Build/Receipts/UE4Game-Win32-Shipping_XP.target.xml
--- a/Engine/Build/Receipts/UE4Game-Win32-Shipping_XP.target.xml
+++ b/Engine/Build/Receipts/UE4Game-Win32-Shipping_XP.target.xml
@@ -1,5 +1,11 @@
 <?xml version="1.0" encoding="utf-8"?>
 <BuildReceipt xmlns:xsd="http://www.w3.org/2001/XMLSchema" xmlns:xsi="http://www.w3.org/2001/XMLSchema-instance">
+  <Properties>
+    <Property Name="TargetName" Value="UE4Game" />
+    <Property Name="Platform" Value="Win32" />
+    <Property Name="Configuration" Value="Shipping" />
+    <Property Name="Precompile" Value="False" />
+  </Properties>
   <BuildProducts>
     <BuildProduct Path="$(EngineDir)\Binaries\Win32\UE4Game-Win32-Shipping_XP.exe" Type="Executable" />
     <BuildProduct Path="$(EngineDir)\Binaries\Win32\UE4Game-Win32-Shipping_XP.pdb" Type="SymbolFile" />
diff --git a/Engine/Build/Receipts/UE4Game-Win32-Test.target.xml b/Engine/Build/Receipts/UE4Game-Win32-Test.target.xml
--- a/Engine/Build/Receipts/UE4Game-Win32-Test.target.xml
+++ b/Engine/Build/Receipts/UE4Game-Win32-Test.target.xml
@@ -1,5 +1,11 @@
 <?xml version="1.0" encoding="utf-8"?>
 <BuildReceipt xmlns:xsd="http://www.w3.org/2001/XMLSchema" xmlns:xsi="http://www.w3.org/2001/XMLSchema-instance">
+  <Properties>
+    <Property Name="TargetName" Value="UE4Game" />
+    <Property Name="Platform" Value="Win32" />
+    <Property Name="Configuration" Value="Test" />
+    <Property Name="Precompile" Value="False" />
+  </Properties>
   <BuildProducts>
     <BuildProduct Path="$(EngineDir)\Binaries\Win32\UE4Game-Win32-Test.exe" Type="Executable" />
     <BuildProduct Path="$(EngineDir)\Binaries\Win32\UE4Game-Win32-Test.pdb" Type="SymbolFile" />
diff --git a/Engine/Build/Receipts/UE4Game-Win32-Test_XP.target.xml b/Engine/Build/Receipts/UE4Game-Win32-Test_XP.target.xml
--- a/Engine/Build/Receipts/UE4Game-Win32-Test_XP.target.xml
+++ b/Engine/Build/Receipts/UE4Game-Win32-Test_XP.target.xml
@@ -1,5 +1,11 @@
 <?xml version="1.0" encoding="utf-8"?>
 <BuildReceipt xmlns:xsd="http://www.w3.org/2001/XMLSchema" xmlns:xsi="http://www.w3.org/2001/XMLSchema-instance">
+  <Properties>
+    <Property Name="TargetName" Value="UE4Game" />
+    <Property Name="Platform" Value="Win32" />
+    <Property Name="Configuration" Value="Test" />
+    <Property Name="Precompile" Value="False" />
+  </Properties>
   <BuildProducts>
     <BuildProduct Path="$(EngineDir)\Binaries\Win32\UE4Game-Win32-Test_XP.exe" Type="Executable" />
     <BuildProduct Path="$(EngineDir)\Binaries\Win32\UE4Game-Win32-Test_XP.pdb" Type="SymbolFile" />
diff --git a/Engine/Build/Receipts/UE4Game-Win64-Development.target.xml b/Engine/Build/Receipts/UE4Game-Win64-Development.target.xml
--- a/Engine/Build/Receipts/UE4Game-Win64-Development.target.xml
+++ b/Engine/Build/Receipts/UE4Game-Win64-Development.target.xml
@@ -1,5 +1,11 @@
 <?xml version="1.0" encoding="utf-8"?>
 <BuildReceipt xmlns:xsd="http://www.w3.org/2001/XMLSchema" xmlns:xsi="http://www.w3.org/2001/XMLSchema-instance">
+  <Properties>
+    <Property Name="TargetName" Value="UE4Game" />
+    <Property Name="Platform" Value="Win64" />
+    <Property Name="Configuration" Value="Development" />
+    <Property Name="Precompile" Value="True" />
+  </Properties>
   <BuildProducts>
     <BuildProduct Path="$(EngineDir)\Binaries\Win64\UE4Game.exe" Type="Executable" />
     <BuildProduct Path="$(EngineDir)\Binaries\Win64\UE4Game.pdb" Type="SymbolFile" />
diff --git a/Engine/Build/Receipts/UE4Game-Win64-Shipping.target.xml b/Engine/Build/Receipts/UE4Game-Win64-Shipping.target.xml
--- a/Engine/Build/Receipts/UE4Game-Win64-Shipping.target.xml
+++ b/Engine/Build/Receipts/UE4Game-Win64-Shipping.target.xml
@@ -1,5 +1,11 @@
 <?xml version="1.0" encoding="utf-8"?>
 <BuildReceipt xmlns:xsd="http://www.w3.org/2001/XMLSchema" xmlns:xsi="http://www.w3.org/2001/XMLSchema-instance">
+  <Properties>
+    <Property Name="TargetName" Value="UE4Game" />
+    <Property Name="Platform" Value="Win64" />
+    <Property Name="Configuration" Value="Shipping" />
+    <Property Name="Precompile" Value="True" />
+  </Properties>
   <BuildProducts>
     <BuildProduct Path="$(EngineDir)\Binaries\Win64\UE4Game-Win64-Shipping.exe" Type="Executable" />
     <BuildProduct Path="$(EngineDir)\Binaries\Win64\UE4Game-Win64-Shipping.pdb" Type="SymbolFile" />
diff --git a/Engine/Build/Receipts/UE4Game-Win64-Test.target.xml b/Engine/Build/Receipts/UE4Game-Win64-Test.target.xml
--- a/Engine/Build/Receipts/UE4Game-Win64-Test.target.xml
+++ b/Engine/Build/Receipts/UE4Game-Win64-Test.target.xml
@@ -1,5 +1,11 @@
 <?xml version="1.0" encoding="utf-8"?>
 <BuildReceipt xmlns:xsd="http://www.w3.org/2001/XMLSchema" xmlns:xsi="http://www.w3.org/2001/XMLSchema-instance">
+  <Properties>
+    <Property Name="TargetName" Value="UE4Game" />
+    <Property Name="Platform" Value="Win64" />
+    <Property Name="Configuration" Value="Test" />
+    <Property Name="Precompile" Value="False" />
+  </Properties>
   <BuildProducts>
     <BuildProduct Path="$(EngineDir)\Binaries\Win64\UE4Game-Win64-Test.exe" Type="Executable" />
     <BuildProduct Path="$(EngineDir)\Binaries\Win64\UE4Game-Win64-Test.pdb" Type="SymbolFile" />
@@ -11,7 +17,6 @@
     <BuildProduct Path="$(EngineDir)\Plugins\Runtime\CustomMeshComponent\Binaries\Win64\UE4Game-CustomMeshComponent-Win64-Test-Static.lib" Type="StaticLibrary" />
     <BuildProduct Path="$(EngineDir)\Plugins\Runtime\ExampleDeviceProfileSelector\Binaries\Win64\UE4Game-ExampleDeviceProfileSelector-Win64-Test-Static.lib" Type="StaticLibrary" />
     <BuildProduct Path="$(EngineDir)\Plugins\Runtime\OculusRift\Binaries\Win64\UE4Game-OculusRift-Win64-Test-Static.lib" Type="StaticLibrary" />
-    <BuildProduct Path="$(EngineDir)\Plugins\Runtime\PS4\Morpheus\Binaries\Win64\UE4Game-Morpheus-Win64-Test-Static.lib" Type="StaticLibrary" />
     <BuildProduct Path="$(EngineDir)\Plugins\Runtime\WindowsMoviePlayer\Binaries\Win64\UE4Game-WindowsMoviePlayer-Win64-Test-Static.lib" Type="StaticLibrary" />
     <BuildProduct Path="$(EngineDir)\Plugins\Slate\SlateRemote\Binaries\Win64\UE4Game-SlateRemote-Win64-Test-Static.lib" Type="StaticLibrary" />
   </BuildProducts>
diff --git a/Engine/Build/Receipts/UE4Game-XboxOne-Development.target.xml b/Engine/Build/Receipts/UE4Game-XboxOne-Development.target.xml
--- a/Engine/Build/Receipts/UE4Game-XboxOne-Development.target.xml
+++ b/Engine/Build/Receipts/UE4Game-XboxOne-Development.target.xml
@@ -1,5 +1,11 @@
 <?xml version="1.0" encoding="utf-8"?>
 <BuildReceipt xmlns:xsd="http://www.w3.org/2001/XMLSchema" xmlns:xsi="http://www.w3.org/2001/XMLSchema-instance">
+  <Properties>
+    <Property Name="TargetName" Value="UE4Game" />
+    <Property Name="Platform" Value="XboxOne" />
+    <Property Name="Configuration" Value="Development" />
+    <Property Name="Precompile" Value="False" />
+  </Properties>
   <BuildProducts>
     <BuildProduct Path="$(EngineDir)\Binaries\XboxOne\UE4Game.exe" Type="Executable" />
     <BuildProduct Path="$(EngineDir)\Binaries\XboxOne\UE4Game.pdb" Type="SymbolFile" />
diff --git a/Engine/Build/Receipts/UE4Game-XboxOne-Shipping.target.xml b/Engine/Build/Receipts/UE4Game-XboxOne-Shipping.target.xml
--- a/Engine/Build/Receipts/UE4Game-XboxOne-Shipping.target.xml
+++ b/Engine/Build/Receipts/UE4Game-XboxOne-Shipping.target.xml
@@ -1,5 +1,11 @@
 <?xml version="1.0" encoding="utf-8"?>
 <BuildReceipt xmlns:xsd="http://www.w3.org/2001/XMLSchema" xmlns:xsi="http://www.w3.org/2001/XMLSchema-instance">
+  <Properties>
+    <Property Name="TargetName" Value="UE4Game" />
+    <Property Name="Platform" Value="XboxOne" />
+    <Property Name="Configuration" Value="Shipping" />
+    <Property Name="Precompile" Value="False" />
+  </Properties>
   <BuildProducts>
     <BuildProduct Path="$(EngineDir)\Binaries\XboxOne\UE4Game-XboxOne-Shipping.exe" Type="Executable" />
     <BuildProduct Path="$(EngineDir)\Binaries\XboxOne\UE4Game-XboxOne-Shipping.pdb" Type="SymbolFile" />
diff --git a/Engine/Build/Receipts/UE4Game-XboxOne-Test.target.xml b/Engine/Build/Receipts/UE4Game-XboxOne-Test.target.xml
--- a/Engine/Build/Receipts/UE4Game-XboxOne-Test.target.xml
+++ b/Engine/Build/Receipts/UE4Game-XboxOne-Test.target.xml
@@ -1,5 +1,11 @@
 <?xml version="1.0" encoding="utf-8"?>
 <BuildReceipt xmlns:xsd="http://www.w3.org/2001/XMLSchema" xmlns:xsi="http://www.w3.org/2001/XMLSchema-instance">
+  <Properties>
+    <Property Name="TargetName" Value="UE4Game" />
+    <Property Name="Platform" Value="XboxOne" />
+    <Property Name="Configuration" Value="Test" />
+    <Property Name="Precompile" Value="False" />
+  </Properties>
   <BuildProducts>
     <BuildProduct Path="$(EngineDir)\Binaries\XboxOne\UE4Game-XboxOne-Test.exe" Type="Executable" />
     <BuildProduct Path="$(EngineDir)\Binaries\XboxOne\UE4Game-XboxOne-Test.pdb" Type="SymbolFile" />
diff --git a/Engine/Build/Receipts/UE4Server-Linux-Developmentx86_64-unknown-linux-gnu.target.xml b/Engine/Build/Receipts/UE4Server-Linux-Developmentx86_64-unknown-linux-gnu.target.xml
--- a/Engine/Build/Receipts/UE4Server-Linux-Developmentx86_64-unknown-linux-gnu.target.xml
+++ b/Engine/Build/Receipts/UE4Server-Linux-Developmentx86_64-unknown-linux-gnu.target.xml
@@ -1,5 +1,11 @@
 <?xml version="1.0" encoding="utf-8"?>
 <BuildReceipt xmlns:xsd="http://www.w3.org/2001/XMLSchema" xmlns:xsi="http://www.w3.org/2001/XMLSchema-instance">
+  <Properties>
+    <Property Name="TargetName" Value="UE4Server" />
+    <Property Name="Platform" Value="Linux" />
+    <Property Name="Configuration" Value="Development" />
+    <Property Name="Precompile" Value="False" />
+  </Properties>
   <BuildProducts>
     <BuildProduct Path="$(EngineDir)\Binaries\Linux\UE4Server" Type="Executable" />
     <BuildProduct Path="$(EngineDir)\Plugins\2D\Paper2D\Binaries\Linux\libUE4Server-Paper2D-Static.a" Type="StaticLibrary" />
diff --git a/Engine/Build/Receipts/UE4Server-Win32-Development.target.xml b/Engine/Build/Receipts/UE4Server-Win32-Development.target.xml
--- a/Engine/Build/Receipts/UE4Server-Win32-Development.target.xml
+++ b/Engine/Build/Receipts/UE4Server-Win32-Development.target.xml
@@ -1,5 +1,11 @@
 <?xml version="1.0" encoding="utf-8"?>
 <BuildReceipt xmlns:xsd="http://www.w3.org/2001/XMLSchema" xmlns:xsi="http://www.w3.org/2001/XMLSchema-instance">
+  <Properties>
+    <Property Name="TargetName" Value="UE4Server" />
+    <Property Name="Platform" Value="Win32" />
+    <Property Name="Configuration" Value="Development" />
+    <Property Name="Precompile" Value="False" />
+  </Properties>
   <BuildProducts>
     <BuildProduct Path="$(EngineDir)\Binaries\Win32\UE4Server.exe" Type="Executable" />
     <BuildProduct Path="$(EngineDir)\Binaries\Win32\UE4Server.pdb" Type="SymbolFile" />
diff --git a/Engine/Build/Receipts/UE4Server-Win64-Development.target.xml b/Engine/Build/Receipts/UE4Server-Win64-Development.target.xml
--- a/Engine/Build/Receipts/UE4Server-Win64-Development.target.xml
+++ b/Engine/Build/Receipts/UE4Server-Win64-Development.target.xml
@@ -1,5 +1,11 @@
 <?xml version="1.0" encoding="utf-8"?>
 <BuildReceipt xmlns:xsd="http://www.w3.org/2001/XMLSchema" xmlns:xsi="http://www.w3.org/2001/XMLSchema-instance">
+  <Properties>
+    <Property Name="TargetName" Value="UE4Server" />
+    <Property Name="Platform" Value="Win64" />
+    <Property Name="Configuration" Value="Development" />
+    <Property Name="Precompile" Value="False" />
+  </Properties>
   <BuildProducts>
     <BuildProduct Path="$(EngineDir)\Binaries\Win64\UE4Server.exe" Type="Executable" />
     <BuildProduct Path="$(EngineDir)\Binaries\Win64\UE4Server.pdb" Type="SymbolFile" />
@@ -11,7 +17,6 @@
     <BuildProduct Path="$(EngineDir)\Plugins\Runtime\CustomMeshComponent\Binaries\Win64\UE4Server-CustomMeshComponent-Static.lib" Type="StaticLibrary" />
     <BuildProduct Path="$(EngineDir)\Plugins\Runtime\ExampleDeviceProfileSelector\Binaries\Win64\UE4Server-ExampleDeviceProfileSelector-Static.lib" Type="StaticLibrary" />
     <BuildProduct Path="$(EngineDir)\Plugins\Runtime\OculusRift\Binaries\Win64\UE4Server-OculusRift-Static.lib" Type="StaticLibrary" />
-    <BuildProduct Path="$(EngineDir)\Plugins\Runtime\PS4\Morpheus\Binaries\Win64\UE4Server-Morpheus-Static.lib" Type="StaticLibrary" />
     <BuildProduct Path="$(EngineDir)\Plugins\Runtime\WindowsMoviePlayer\Binaries\Win64\UE4Server-WindowsMoviePlayer-Static.lib" Type="StaticLibrary" />
     <BuildProduct Path="$(EngineDir)\Plugins\Slate\SlateRemote\Binaries\Win64\UE4Server-SlateRemote-Static.lib" Type="StaticLibrary" />
   </BuildProducts>
diff --git a/Engine/Build/Receipts/UnrealCEFSubProcess-Linux-Developmentx86_64-unknown-linux-gnu.target.xml b/Engine/Build/Receipts/UnrealCEFSubProcess-Linux-Developmentx86_64-unknown-linux-gnu.target.xml
new file mode 100644
--- /dev/null
+++ b/Engine/Build/Receipts/UnrealCEFSubProcess-Linux-Developmentx86_64-unknown-linux-gnu.target.xml
@@ -0,0 +1,13 @@
+<?xml version="1.0" encoding="utf-8"?>
+<BuildReceipt xmlns:xsd="http://www.w3.org/2001/XMLSchema" xmlns:xsi="http://www.w3.org/2001/XMLSchema-instance">
+  <Properties>
+    <Property Name="TargetName" Value="UnrealCEFSubProcess" />
+    <Property Name="Platform" Value="Linux" />
+    <Property Name="Configuration" Value="Development" />
+    <Property Name="Precompile" Value="False" />
+  </Properties>
+  <BuildProducts>
+    <BuildProduct Path="$(EngineDir)\Binaries\Linux\UnrealCEFSubProcess" Type="Executable" />
+  </BuildProducts>
+  <RuntimeDependencies />
+</BuildReceipt>
\ No newline at end of file
diff --git a/Engine/Build/Receipts/UnrealFrontend-Linux-Developmentx86_64-unknown-linux-gnu.target.xml b/Engine/Build/Receipts/UnrealFrontend-Linux-Developmentx86_64-unknown-linux-gnu.target.xml
new file mode 100644
--- /dev/null
+++ b/Engine/Build/Receipts/UnrealFrontend-Linux-Developmentx86_64-unknown-linux-gnu.target.xml
@@ -0,0 +1,70 @@
+<?xml version="1.0" encoding="utf-8"?>
+<BuildReceipt xmlns:xsd="http://www.w3.org/2001/XMLSchema" xmlns:xsi="http://www.w3.org/2001/XMLSchema-instance">
+  <Properties>
+    <Property Name="TargetName" Value="UnrealFrontend" />
+    <Property Name="Platform" Value="Linux" />
+    <Property Name="Configuration" Value="Development" />
+    <Property Name="Precompile" Value="False" />
+  </Properties>
+  <BuildProducts>
+    <BuildProduct Path="$(EngineDir)\Binaries\Linux\UnrealFrontend" Type="Executable" />
+    <BuildProduct Path="$(EngineDir)\Plugins\Messaging\UdpMessaging\Binaries\Linux\libUnrealFrontend-UdpMessaging.so" Type="DynamicLibrary" />
+    <BuildProduct Path="$(EngineDir)\Binaries\Linux\libUnrealFrontend-AutomationController.so" Type="DynamicLibrary" />
+    <BuildProduct Path="$(EngineDir)\Binaries\Linux\libUnrealFrontend-AutomationMessages.so" Type="DynamicLibrary" />
+    <BuildProduct Path="$(EngineDir)\Binaries\Linux\libUnrealFrontend-Core.so" Type="DynamicLibrary" />
+    <BuildProduct Path="$(EngineDir)\Binaries\Linux\Linux\libUnrealFrontend-LinuxTargetPlatform.so" Type="DynamicLibrary" />
+    <BuildProduct Path="$(EngineDir)\Binaries\Linux\libUnrealFrontend-CoreUObject.so" Type="DynamicLibrary" />
+    <BuildProduct Path="$(EngineDir)\Binaries\Linux\libUnrealFrontend-Projects.so" Type="DynamicLibrary" />
+    <BuildProduct Path="$(EngineDir)\Binaries\Linux\libUnrealFrontend-Json.so" Type="DynamicLibrary" />
+    <BuildProduct Path="$(EngineDir)\Binaries\Linux\Linux\libUnrealFrontend-LinuxNoEditorTargetPlatform.so" Type="DynamicLibrary" />
+    <BuildProduct Path="$(EngineDir)\Binaries\Linux\libUnrealFrontend-TargetPlatform.so" Type="DynamicLibrary" />
+    <BuildProduct Path="$(EngineDir)\Binaries\Linux\libUnrealFrontend-DesktopPlatform.so" Type="DynamicLibrary" />
+    <BuildProduct Path="$(EngineDir)\Binaries\Linux\Linux\libUnrealFrontend-LinuxServerTargetPlatform.so" Type="DynamicLibrary" />
+    <BuildProduct Path="$(EngineDir)\Binaries\Linux\libUnrealFrontend-DesktopTargetPlatform.so" Type="DynamicLibrary" />
+    <BuildProduct Path="$(EngineDir)\Binaries\Linux\libUnrealFrontend-LaunchDaemonMessages.so" Type="DynamicLibrary" />
+    <BuildProduct Path="$(EngineDir)\Binaries\Linux\libUnrealFrontend-ShaderPreprocessor.so" Type="DynamicLibrary" />
+    <BuildProduct Path="$(EngineDir)\Binaries\Linux\libUnrealFrontend-ShaderCore.so" Type="DynamicLibrary" />
+    <BuildProduct Path="$(EngineDir)\Binaries\Linux\libUnrealFrontend-RHI.so" Type="DynamicLibrary" />
+    <BuildProduct Path="$(EngineDir)\Binaries\Linux\libUnrealFrontend-RenderCore.so" Type="DynamicLibrary" />
+    <BuildProduct Path="$(EngineDir)\Binaries\Linux\libUnrealFrontend-ShaderFormatOpenGL.so" Type="DynamicLibrary" />
+    <BuildProduct Path="$(EngineDir)\Binaries\Linux\libUnrealFrontend-ShaderCompilerCommon.so" Type="DynamicLibrary" />
+    <BuildProduct Path="$(EngineDir)\Binaries\Linux\libUnrealFrontend-ImageWrapper.so" Type="DynamicLibrary" />
+    <BuildProduct Path="$(EngineDir)\Binaries\Linux\libUnrealFrontend-TextureFormatDXT.so" Type="DynamicLibrary" />
+    <BuildProduct Path="$(EngineDir)\Binaries\Linux\libUnrealFrontend-ImageCore.so" Type="DynamicLibrary" />
+    <BuildProduct Path="$(EngineDir)\Binaries\Linux\libUnrealFrontend-TextureFormatPVR.so" Type="DynamicLibrary" />
+    <BuildProduct Path="$(EngineDir)\Binaries\Linux\libUnrealFrontend-TextureFormatASTC.so" Type="DynamicLibrary" />
+    <BuildProduct Path="$(EngineDir)\Binaries\Linux\libUnrealFrontend-TextureFormatUncompressed.so" Type="DynamicLibrary" />
+    <BuildProduct Path="$(EngineDir)\Binaries\Linux\libUnrealFrontend-UnrealEdMessages.so" Type="DynamicLibrary" />
+    <BuildProduct Path="$(EngineDir)\Binaries\Linux\libUnrealFrontend-DeviceManager.so" Type="DynamicLibrary" />
+    <BuildProduct Path="$(EngineDir)\Binaries\Linux\libUnrealFrontend-EditorStyle.so" Type="DynamicLibrary" />
+    <BuildProduct Path="$(EngineDir)\Binaries\Linux\libUnrealFrontend-SlateCore.so" Type="DynamicLibrary" />
+    <BuildProduct Path="$(EngineDir)\Binaries\Linux\libUnrealFrontend-InputCore.so" Type="DynamicLibrary" />
+    <BuildProduct Path="$(EngineDir)\Binaries\Linux\libUnrealFrontend-Slate.so" Type="DynamicLibrary" />
+    <BuildProduct Path="$(EngineDir)\Binaries\Linux\libUnrealFrontend-WorkspaceMenuStructure.so" Type="DynamicLibrary" />
+    <BuildProduct Path="$(EngineDir)\Binaries\Linux\libUnrealFrontend-LauncherServices.so" Type="DynamicLibrary" />
+    <BuildProduct Path="$(EngineDir)\Binaries\Linux\libUnrealFrontend-SessionMessages.so" Type="DynamicLibrary" />
+    <BuildProduct Path="$(EngineDir)\Binaries\Linux\libUnrealFrontend-SourceCodeAccess.so" Type="DynamicLibrary" />
+    <BuildProduct Path="$(EngineDir)\Binaries\Linux\libUnrealFrontend-Settings.so" Type="DynamicLibrary" />
+    <BuildProduct Path="$(EngineDir)\Binaries\Linux\libUnrealFrontend-Messaging.so" Type="DynamicLibrary" />
+    <BuildProduct Path="$(EngineDir)\Binaries\Linux\libUnrealFrontend-ProfilerClient.so" Type="DynamicLibrary" />
+    <BuildProduct Path="$(EngineDir)\Binaries\Linux\libUnrealFrontend-ProfilerMessages.so" Type="DynamicLibrary" />
+    <BuildProduct Path="$(EngineDir)\Binaries\Linux\libUnrealFrontend-ProjectLauncher.so" Type="DynamicLibrary" />
+    <BuildProduct Path="$(EngineDir)\Binaries\Linux\libUnrealFrontend-SessionFrontend.so" Type="DynamicLibrary" />
+    <BuildProduct Path="$(EngineDir)\Binaries\Linux\libUnrealFrontend-SessionServices.so" Type="DynamicLibrary" />
+    <BuildProduct Path="$(EngineDir)\Binaries\Linux\libUnrealFrontend-EngineMessages.so" Type="DynamicLibrary" />
+    <BuildProduct Path="$(EngineDir)\Binaries\Linux\libUnrealFrontend-AutomationWindow.so" Type="DynamicLibrary" />
+    <BuildProduct Path="$(EngineDir)\Binaries\Linux\libUnrealFrontend-ScreenShotComparison.so" Type="DynamicLibrary" />
+    <BuildProduct Path="$(EngineDir)\Binaries\Linux\libUnrealFrontend-ScreenShotComparisonTools.so" Type="DynamicLibrary" />
+    <BuildProduct Path="$(EngineDir)\Binaries\Linux\libUnrealFrontend-Profiler.so" Type="DynamicLibrary" />
+    <BuildProduct Path="$(EngineDir)\Binaries\Linux\libUnrealFrontend-SlateReflector.so" Type="DynamicLibrary" />
+    <BuildProduct Path="$(EngineDir)\Binaries\Linux\libUnrealFrontend-StandaloneRenderer.so" Type="DynamicLibrary" />
+    <BuildProduct Path="$(EngineDir)\Binaries\Linux\libUnrealFrontend-TargetDeviceServices.so" Type="DynamicLibrary" />
+    <BuildProduct Path="$(EngineDir)\Binaries\Linux\libUnrealFrontend-Networking.so" Type="DynamicLibrary" />
+    <BuildProduct Path="$(EngineDir)\Binaries\Linux\libUnrealFrontend-Sockets.so" Type="DynamicLibrary" />
+    <BuildProduct Path="$(EngineDir)\Binaries\Linux\libUnrealFrontend-SandboxFile.so" Type="DynamicLibrary" />
+    <BuildProduct Path="$(EngineDir)\Binaries\Linux\libUnrealFrontend-DerivedDataCache.so" Type="DynamicLibrary" />
+    <BuildProduct Path="$(EngineDir)\Binaries\Linux\NotForLicensees\libUnrealFrontend-DDCUtils.so" Type="DynamicLibrary" />
+    <BuildProduct Path="$(EngineDir)\Binaries\Linux\libUnrealFrontend-Serialization.so" Type="DynamicLibrary" />
+  </BuildProducts>
+  <RuntimeDependencies />
+</BuildReceipt>
\ No newline at end of file
diff --git a/Engine/Build/Receipts/UnrealLightmass-Linux-Developmentx86_64-unknown-linux-gnu.target.xml b/Engine/Build/Receipts/UnrealLightmass-Linux-Developmentx86_64-unknown-linux-gnu.target.xml
--- a/Engine/Build/Receipts/UnrealLightmass-Linux-Developmentx86_64-unknown-linux-gnu.target.xml
+++ b/Engine/Build/Receipts/UnrealLightmass-Linux-Developmentx86_64-unknown-linux-gnu.target.xml
@@ -1,5 +1,11 @@
 <?xml version="1.0" encoding="utf-8"?>
 <BuildReceipt xmlns:xsd="http://www.w3.org/2001/XMLSchema" xmlns:xsi="http://www.w3.org/2001/XMLSchema-instance">
+  <Properties>
+    <Property Name="TargetName" Value="UnrealLightmass" />
+    <Property Name="Platform" Value="Linux" />
+    <Property Name="Configuration" Value="Development" />
+    <Property Name="Precompile" Value="False" />
+  </Properties>
   <BuildProducts>
     <BuildProduct Path="$(EngineDir)\Binaries\Linux\UnrealLightmass" Type="Executable" />
     <BuildProduct Path="$(EngineDir)\Plugins\Messaging\UdpMessaging\Binaries\Linux\libUnrealLightmass-UdpMessaging.so" Type="DynamicLibrary" />
diff --git a/Engine/Build/Receipts/UnrealPak-Linux-Developmentx86_64-unknown-linux-gnu.target.xml b/Engine/Build/Receipts/UnrealPak-Linux-Developmentx86_64-unknown-linux-gnu.target.xml
new file mode 100644
--- /dev/null
+++ b/Engine/Build/Receipts/UnrealPak-Linux-Developmentx86_64-unknown-linux-gnu.target.xml
@@ -0,0 +1,13 @@
+<?xml version="1.0" encoding="utf-8"?>
+<BuildReceipt xmlns:xsd="http://www.w3.org/2001/XMLSchema" xmlns:xsi="http://www.w3.org/2001/XMLSchema-instance">
+  <Properties>
+    <Property Name="TargetName" Value="UnrealPak" />
+    <Property Name="Platform" Value="Linux" />
+    <Property Name="Configuration" Value="Development" />
+    <Property Name="Precompile" Value="False" />
+  </Properties>
+  <BuildProducts>
+    <BuildProduct Path="$(EngineDir)\Binaries\Linux\UnrealPak" Type="Executable" />
+  </BuildProducts>
+  <RuntimeDependencies />
+</BuildReceipt>
\ No newline at end of file
diff --git a/Engine/Build/Receipts/XboxOnePDBFileUtil-Win64-Development.target.xml b/Engine/Build/Receipts/XboxOnePDBFileUtil-Win64-Development.target.xml
--- a/Engine/Build/Receipts/XboxOnePDBFileUtil-Win64-Development.target.xml
+++ b/Engine/Build/Receipts/XboxOnePDBFileUtil-Win64-Development.target.xml
@@ -1,5 +1,11 @@
 <?xml version="1.0" encoding="utf-8"?>
 <BuildReceipt xmlns:xsd="http://www.w3.org/2001/XMLSchema" xmlns:xsi="http://www.w3.org/2001/XMLSchema-instance">
+  <Properties>
+    <Property Name="TargetName" Value="XboxOnePDBFileUtil" />
+    <Property Name="Platform" Value="Win64" />
+    <Property Name="Configuration" Value="Development" />
+    <Property Name="Precompile" Value="False" />
+  </Properties>
   <BuildProducts>
     <BuildProduct Path="$(EngineDir)\Binaries\Win64\XboxOne\XboxOnePDBFileUtil.exe" Type="Executable" />
     <BuildProduct Path="$(EngineDir)\Binaries\Win64\XboxOne\XboxOnePDBFileUtil.pdb" Type="SymbolFile" />
